Change presave prototyp: Deprecation fixes for UE5 project.
#rnx

diff --git a/Engine/Plugins/Experimental/CommonConversation/Source/CommonConversationGraph/Public/CommonConversationGraphModule.h b/Engine/Plugins/Experimental/CommonConversation/Source/CommonConversationGraph/Public/CommonConversationGraphModule.h
--- a/Engine/Plugins/Experimental/CommonConversation/Source/CommonConversationGraph/Public/CommonConversationGraphModule.h
+++ b/Engine/Plugins/Experimental/CommonConversation/Source/CommonConversationGraph/Public/CommonConversationGraphModule.h
@@ -5,6 +5,7 @@
 #include "CoreMinimal.h"
 #include "Modules/ModuleManager.h"
 
+class FObjectPreSaveContext;
 class UObject;
 class UPackage;
 
@@ -19,6 +20,6 @@ public:
 	virtual void ShutdownModule() override;
 
 private:
-	void HandlePreSavePackage(UPackage* Package);
+	void HandlePreSavePackage(UPackage* Package, FObjectPreSaveContext ObjectSaveContext);
 	void HandleBeginPIE(bool bIsSimulating);
 };
