Conversation Instance - Exposing access to CurrentUserChoices to derived conversation instance classes

diff --git a/Engine/Plugins/Experimental/CommonConversation/Source/CommonConversationRuntime/Public/ConversationInstance.h b/Engine/Plugins/Experimental/CommonConversation/Source/CommonConversationRuntime/Public/ConversationInstance.h
--- a/Engine/Plugins/Experimental/CommonConversation/Source/CommonConversationRuntime/Public/ConversationInstance.h
+++ b/Engine/Plugins/Experimental/CommonConversation/Source/CommonConversationRuntime/Public/ConversationInstance.h
@@ -113,6 +113,10 @@ private:
 	void OnCurrentConversationNodeModified();
 #endif
 
+protected:
+
+	TArray<FClientConversationOptionEntry> CurrentUserChoices;
+
 private:
 	UPROPERTY()
 	FConversationParticipants Participants;
@@ -131,7 +135,6 @@ private:
 	TArray<FCheckpoint> ClientBranchPoints;
 
 	TArray<FConversationBranchPoint> CurrentBranchPoints;
-	TArray<FClientConversationOptionEntry> CurrentUserChoices;
 
 	TArray<FConversationChoiceReference> ScopeStack;
 
