Conversation Participant Component - ServerAbortAllConversations now public to allow owning actor to abort conversations it is involved with, when external owner specific conditions are met.

diff --git a/Engine/Plugins/Experimental/CommonConversation/Source/CommonConversationRuntime/Public/ConversationParticipantComponent.h b/Engine/Plugins/Experimental/CommonConversation/Source/CommonConversationRuntime/Public/ConversationParticipantComponent.h
--- a/Engine/Plugins/Experimental/CommonConversation/Source/CommonConversationRuntime/Public/ConversationParticipantComponent.h
+++ b/Engine/Plugins/Experimental/CommonConversation/Source/CommonConversationRuntime/Public/ConversationParticipantComponent.h
@@ -54,6 +54,9 @@ public:
 	/** Ask to this actor to change is state to be able to start a conversation */
 	virtual void ServerGetReadyToConverse();
 
+	/** Ask this actor to abort all active conversations */
+	void ServerAbortAllConversations();
+
 	FConversationMemory& GetParticipantMemory() { return ParticipantMemory; }
 
 	DECLARE_MULTICAST_DELEGATE_OneParam(FParticipantReadyToConverse, UConversationParticipantComponent*);
@@ -109,7 +112,6 @@ protected:
 #if WITH_SERVER_CODE
 	UConversationInstance* GetCurrentConversationForAuthority() const { return Auth_CurrentConversation; }
 	const TArray<UConversationInstance*>& GetConversationsForAuthority() const { return Auth_Conversations; }
-	void ServerAbortAllConversations();
 #endif
 
 public:
