Conversations - Conversation registry now leases out conversation handles that are smart enough to re-stream conversation assets if the game feature plugins change (and then change the conversation registry's loaded conversations).  If that happens we correctly rebuild the registry, based on what's active, and restream anything new that's need, and of course since the streaming handle is swapped out, we release anything that would no longer be active.

Jonathan.Bunner

diff --git a/Engine/Plugins/Experimental/CommonConversation/Source/CommonConversationRuntime/Private/ConversationRegistry.cpp b/Engine/Plugins/Experimental/CommonConversation/Source/CommonConversationRuntime/Private/ConversationRegistry.cpp
--- a/Engine/Plugins/Experimental/CommonConversation/Source/CommonConversationRuntime/Private/ConversationRegistry.cpp
+++ b/Engine/Plugins/Experimental/CommonConversation/Source/CommonConversationRuntime/Private/ConversationRegistry.cpp
@@ -15,6 +15,41 @@
 
 //======================================================================================
 
+TSharedPtr<FConversationsHandle> FConversationsHandle::Create(const UConversationRegistry* InOwningRegistry, const TSharedPtr<FStreamableHandle>& InStreamableHandle, const TArray<FGameplayTag>& InEntryTags)
+{
+	TSharedPtr<FConversationsHandle> ConversationsHandle = MakeShared<FConversationsHandle>(const_cast<UConversationRegistry*>(InOwningRegistry), InStreamableHandle, InEntryTags);
+	ConversationsHandle->Initialize();
+
+	return ConversationsHandle;
+}
+
+FConversationsHandle::FConversationsHandle(UConversationRegistry* InOwningRegistry, const TSharedPtr<FStreamableHandle>& InStreamableHandle, const TArray<FGameplayTag>& InEntryTags)
+	: StreamableHandle(InStreamableHandle)
+	, ConversationEntryTags(InEntryTags)
+	, OwningRegistryPtr(InOwningRegistry)
+{
+}
+
+void FConversationsHandle::Initialize()
+{
+	if (UConversationRegistry* OwningRegistry = OwningRegistryPtr.Get())
+	{
+		OwningRegistry->AvailableConversationsChanged.AddSP(this, &FConversationsHandle::HandleAvailableConversationsChanged);
+	}
+}
+
+void FConversationsHandle::HandleAvailableConversationsChanged()
+{
+	if (const UConversationRegistry* OwningRegistry = OwningRegistryPtr.Get())
+	{
+		// Requery the conversations, and swap the streamable handle out.
+		TSharedPtr<FConversationsHandle> TempConversation = OwningRegistry->LoadConversationsFor(ConversationEntryTags);
+		StreamableHandle = TempConversation->StreamableHandle;
+	}
+}
+
+//======================================================================================
+
 void FNetSerializeScriptStructCache_ConvVersion::InitForType(UScriptStruct* InScriptStruct)
 {
 	ScriptStructsToIndex.Reset();
@@ -116,6 +151,32 @@ UConversationRegistry* UConversationRegistry::GetFromWorld(const UWorld* World)
 	return UWorld::GetSubsystem<UConversationRegistry>(World);
 }
 
+void UConversationRegistry::Initialize(FSubsystemCollectionBase& Collection)
+{
+	Super::Initialize(Collection);
+
+	UGameFeaturesSubsystem::OnPluginLoadCompleteDataReady().AddUObject(this, &UConversationRegistry::HandlePluginActivationChanged);
+	UGameFeaturesSubsystem::OnPluginDeactivatedDataReady().AddUObject(this, &UConversationRegistry::HandlePluginActivationChanged);
+}
+
+void UConversationRegistry::Deinitialize()
+{
+	Super::Deinitialize();
+
+	UGameFeaturesSubsystem::OnPluginLoadCompleteDataReady().RemoveAll(this);
+	UGameFeaturesSubsystem::OnPluginDeactivatedDataReady().RemoveAll(this);
+}
+
+void UConversationRegistry::HandlePluginActivationChanged(const FString& GameFeatureName, const UGameFeatureData* GameFeatureData)
+{
+	// If nobody has actually built the dependency graph yet, there's no reason to invalidate anything, nobody is using it yet.
+	if (bDependenciesBuilt)
+	{
+		bDependenciesBuilt = false;
+		BuildDependenciesGraph();
+	}
+}
+
 UConversationNode* UConversationRegistry::GetRuntimeNodeFromGUID(const FGuid& NodeGUID) const
 {
 	UE_LOG(LogCommonConversationRuntime, Verbose, TEXT("Start UConversationRegistry::GetRuntimeNodeFromGUID with NodeGUID: (%s)"), *NodeGUID.ToString());
@@ -199,12 +260,12 @@ UConversationDatabase* UConversationRegistry::GetConversationFromNodeGUID(const
 	return nullptr;
 }
 
-TSharedPtr<FStreamableHandle> UConversationRegistry::LoadConversationsFor(const FGameplayTag& ConversationEntryTag) const
+TSharedPtr<FConversationsHandle> UConversationRegistry::LoadConversationsFor(const FGameplayTag& ConversationEntryTag) const
 {
 	return LoadConversationsFor(TArray<FGameplayTag>({ ConversationEntryTag }));
 }
 
-TSharedPtr<FStreamableHandle> UConversationRegistry::LoadConversationsFor(const TArray<FGameplayTag>& ConversationEntryTags) const
+TSharedPtr<FConversationsHandle> UConversationRegistry::LoadConversationsFor(const TArray<FGameplayTag>& ConversationEntryTags) const
 {
 	const_cast<UConversationRegistry*>(this)->BuildDependenciesGraph();
 
@@ -235,7 +296,8 @@ TSharedPtr<FStreamableHandle> UConversationRegistry::LoadConversationsFor(const
 			*FString::JoinBy(ConversationsToLoad, TEXT(", "), [](const FSoftObjectPath& SoftObjectPath) { return FString::Printf(TEXT("'%s'"), *SoftObjectPath.ToString()); })
 		);
 
-		return UAssetManager::Get().LoadAssetList(ConversationsToLoad.Array());
+		TSharedPtr<FStreamableHandle> StreamableHandle = UAssetManager::Get().LoadAssetList(ConversationsToLoad.Array());
+		return FConversationsHandle::Create(this, StreamableHandle, ConversationEntryTags);
 	}
 	else
 	{
@@ -244,7 +306,7 @@ TSharedPtr<FStreamableHandle> UConversationRegistry::LoadConversationsFor(const
 		);
 	}
 
-	return TSharedPtr<FStreamableHandle>();
+	return FConversationsHandle::Create(this, TSharedPtr<FStreamableHandle>(), ConversationEntryTags);
 }
 
 TArray<FPrimaryAssetId> UConversationRegistry::GetPrimaryAssetIdsForEntryPoint(FGameplayTag EntryPoint) const
@@ -297,20 +359,45 @@ void UConversationRegistry::BuildDependenciesGraph()
 
 	UE_LOG(LogCommonConversationRuntime, Verbose, TEXT("Registry Building Graph"));
 
-	TArray<FAssetData> AllConversations;
-	UAssetManager::Get().GetPrimaryAssetDataList(FPrimaryAssetType(UConversationDatabase::StaticClass()->GetFName()), AllConversations);
+	TArray<FAssetData> AllActiveConversationAssets;
+	UAssetManager::Get().GetPrimaryAssetDataList(FPrimaryAssetType(UConversationDatabase::StaticClass()->GetFName()), AllActiveConversationAssets);
 
-	// Don't index conversation graphs from inactive game feature plugins.
-	UGameFeaturesSubsystem::Get().FilterInactivePluginAssets(AllConversations);
+	// Don't index conversation graphs from inactive game feature plug-ins.
+	UGameFeaturesSubsystem::Get().FilterInactivePluginAssets(AllActiveConversationAssets);
+
+	// Lets find out if anything actually changed, we maybe rebuilding things after a game feature added or got removed,
+	// as far as conversations are concerned it's possible nothing has changed.
+	if (RuntimeDependencyGraph.Num() == AllActiveConversationAssets.Num())
+	{
+		bool ConversationAssetsChanged = false;
+		for (const FAssetData& ConversationDataAsset : AllActiveConversationAssets)
+		{
+			if (!RuntimeDependencyGraph.Contains(ConversationDataAsset.ToSoftObjectPath()))
+			{
+				ConversationAssetsChanged = true;
+				break;
+			}
+		}
+
+		// If we need to rebuild the conversation graph, but the conversation graph, wont actually change based on what was
+		// loaded or unloaded since the last time BuildDependenciesGraph was called, then we don't actually need to regenerate
+		// anything.
+		if (!ConversationAssetsChanged)
+		{
+			bDependenciesBuilt = true;
+			return;
+		}
+	}
 
 	RuntimeDependencyGraph.Reset();
 	EntryTagToConversations.Reset();
+	EntryTagToEntryList.Reset();
 	NodeGuidToConversation.Reset();
 
-	UE_LOG(LogCommonConversationRuntime, Verbose, TEXT("Building: Total Conversations %d"), AllConversations.Num());
+	UE_LOG(LogCommonConversationRuntime, Verbose, TEXT("Building: Total Conversations %d, Old Conversations %d"), AllActiveConversationAssets.Num(), RuntimeDependencyGraph.Num());
 
 	// Seed
-	for (FAssetData& ConversationDataAsset : AllConversations)
+	for (const FAssetData& ConversationDataAsset : AllActiveConversationAssets)
 	{
 		const FString EntryTagsString = ConversationDataAsset.GetTagValueRef<FString>(GET_MEMBER_NAME_CHECKED(UConversationDatabase, EntryTags));
 		if (!EntryTagsString.IsEmpty())
@@ -346,7 +433,7 @@ void UConversationRegistry::BuildDependenciesGraph()
 	UE_LOG(LogCommonConversationRuntime, Verbose, TEXT("Building: Total Entry Points %d"), EntryTagToConversations.Num());
 	UE_LOG(LogCommonConversationRuntime, Verbose, TEXT("Building: Total Nodes %d"), NodeGuidToConversation.Num());
 
-	for (FAssetData& ConversationDataAsset : AllConversations)
+	for (const FAssetData& ConversationDataAsset : AllActiveConversationAssets)
 	{
 		const FString ExitTagsString = ConversationDataAsset.GetTagValueRef<FString>(GET_MEMBER_NAME_CHECKED(UConversationDatabase, ExitTags));
 		if (!ExitTagsString.IsEmpty())
@@ -387,6 +474,8 @@ void UConversationRegistry::BuildDependenciesGraph()
 	}
 
 	bDependenciesBuilt = true;
+
+	AvailableConversationsChanged.Broadcast();
 }
 
 void UConversationRegistry::GetAllDependenciesForConversation(const FSoftObjectPath& Parent, TSet<FSoftObjectPath>& OutConversationsToLoad) const
diff --git a/Engine/Plugins/Experimental/CommonConversation/Source/CommonConversationRuntime/Public/ConversationRegistry.h b/Engine/Plugins/Experimental/CommonConversation/Source/CommonConversationRuntime/Public/ConversationRegistry.h
--- a/Engine/Plugins/Experimental/CommonConversation/Source/CommonConversationRuntime/Public/ConversationRegistry.h
+++ b/Engine/Plugins/Experimental/CommonConversation/Source/CommonConversationRuntime/Public/ConversationRegistry.h
@@ -10,6 +10,7 @@
 
 #include "ConversationRegistry.generated.h"
 
+class UGameFeatureData;
 class UConversationDatabase;
 class UWorld;
 struct FStreamableHandle;
@@ -32,6 +33,33 @@ struct COMMONCONVERSATIONRUNTIME_API FNetSerializeScriptStructCache_ConvVersion
 	TArray<UScriptStruct*> IndexToScriptStructs;
 };
 
+/**
+ * These handles are issued when someone requests a conversation entry point be streamed in.
+ * As long as this handle remains active, we were continue to keep it keep those elements streamed
+ * in, as well as if new game feature plugins activate, we will stream in additional assets
+ * or let previous ones expire.
+ */
+struct COMMONCONVERSATIONRUNTIME_API FConversationsHandle : public TSharedFromThis<FConversationsHandle>
+{
+	static TSharedPtr<FConversationsHandle> Create(const UConversationRegistry* InOwningRegistry, const TSharedPtr<FStreamableHandle>& InStreamableHandle, const TArray<FGameplayTag>& InEntryTags);
+
+private:
+	FConversationsHandle(UConversationRegistry* InOwningRegistry, const TSharedPtr<FStreamableHandle>& InStreamableHandle, const TArray<FGameplayTag>& InEntryTags);
+
+	void Initialize();
+	void HandleAvailableConversationsChanged();
+
+private:
+	TSharedPtr<FStreamableHandle> StreamableHandle;
+	TArray<FGameplayTag> ConversationEntryTags;
+	TWeakObjectPtr<UConversationRegistry> OwningRegistryPtr;
+
+	friend class SharedPointerInternals::TIntrusiveReferenceController<FConversationsHandle>;
+};
+
+
+DECLARE_MULTICAST_DELEGATE(FAvailableConversationsChangedEvent);
+
 /**
  * A registry that can answer questions about all available dialogue assets
  */
@@ -45,6 +73,11 @@ public:
 
 	static UConversationRegistry* GetFromWorld(const UWorld* World);
 
+	/** UWorldSubsystem Begin */
+	virtual void Initialize(FSubsystemCollectionBase& Collection) override;
+	virtual void Deinitialize() override;
+	/** UWorldSubsystem End */
+
 	UConversationNode* GetRuntimeNodeFromGUID(const FGuid& NodeGUID) const;
 	TArray<FGuid> GetEntryPointGUIDs(FGameplayTag EntryPoint) const;
 
@@ -52,20 +85,24 @@ public:
 	TArray<FGuid> GetOutputLinkGUIDs(const FGuid& SourceGUID) const;
 	TArray<FGuid> GetOutputLinkGUIDs(const TArray<FGuid>& SourceGUIDs) const;
 
-	TSharedPtr<FStreamableHandle> LoadConversationsFor(const FGameplayTag& ConversationEntryTag) const;
-	TSharedPtr<FStreamableHandle> LoadConversationsFor(const TArray<FGameplayTag>& ConversationEntryTags) const;
+	TSharedPtr<FConversationsHandle> LoadConversationsFor(const FGameplayTag& ConversationEntryTag) const;
+	TSharedPtr<FConversationsHandle> LoadConversationsFor(const TArray<FGameplayTag>& ConversationEntryTags) const;
 
 	TArray<FPrimaryAssetId> GetPrimaryAssetIdsForEntryPoint(FGameplayTag EntryPoint) const;
 	
 	UPROPERTY(Transient)
 	FNetSerializeScriptStructCache_ConvVersion ConversationChoiceDataStructCache;
 
+	FAvailableConversationsChangedEvent AvailableConversationsChanged;
+
 private:
 	UConversationDatabase* GetConversationFromNodeGUID(const FGuid& NodeGUID) const;
 
 	void BuildDependenciesGraph();
 	void GetAllDependenciesForConversation(const FSoftObjectPath& Parent, TSet<FSoftObjectPath>& OutConversationsToLoad) const;
 
+	void HandlePluginActivationChanged(const FString& GameFeatureName, const UGameFeatureData* GameFeatureData);
+
 private:
 	TMap<FSoftObjectPath, TArray<FSoftObjectPath>> RuntimeDependencyGraph;
 	TMap<FGameplayTag, TArray<FSoftObjectPath>> EntryTagToConversations;
