Fix for bad alignment arithmetic

diff --git a/Engine/Plugins/Experimental/CommonConversation/Source/CommonConversationRuntime/Public/ConversationMemory.h b/Engine/Plugins/Experimental/CommonConversation/Source/CommonConversationRuntime/Public/ConversationMemory.h
--- a/Engine/Plugins/Experimental/CommonConversation/Source/CommonConversationRuntime/Public/ConversationMemory.h
+++ b/Engine/Plugins/Experimental/CommonConversation/Source/CommonConversationRuntime/Public/ConversationMemory.h
@@ -69,7 +69,7 @@ private:
 
 		void* Allocate(uint64 Size)
 		{
-			uint64 AllocationSize = (Size + (16 - 1)) / 16;
+			uint64 AllocationSize = Align(Size, 16);
 			if (!CurrentSlab || CurrentSlabAllocatedSize + AllocationSize > SlabSize)
 			{
 				TotalAllocatedSize += SlabSize;
