Added a description to the CommonConversation plugin and marked it as experimental

diff --git a/Engine/Plugins/Experimental/CommonConversation/CommonConversation.uplugin b/Engine/Plugins/Experimental/CommonConversation/CommonConversation.uplugin
--- a/Engine/Plugins/Experimental/CommonConversation/CommonConversation.uplugin
+++ b/Engine/Plugins/Experimental/CommonConversation/CommonConversation.uplugin
@@ -3,7 +3,7 @@
 	"Version": 1,
 	"VersionName": "0.1",
 	"FriendlyName": "Common Conversation",
-	"Description": "",
+	"Description": "An *experimental* plugin for authoring graph-based conversation trees",
 	"Category": "Gameplay",
 	"CreatedBy": "Epic Games, Inc.",
 	"CreatedByURL": "http://epicgames.com",
@@ -11,7 +11,7 @@
 	"MarketplaceURL": "",
 	"SupportURL": "",
 	"CanContainContent": false,
-	"IsBetaVersion": false,
+	"IsExperimentalVersion": true,
 	"Installed": false,
 	"EnabledByDefault": false,
 	"Modules": [
