Virtual function for future inherited behavior (SendClientUpdatedChoices)

diff --git a/Engine/Plugins/Experimental/CommonConversation/Source/CommonConversationRuntime/Public/ConversationParticipantComponent.h b/Engine/Plugins/Experimental/CommonConversation/Source/CommonConversationRuntime/Public/ConversationParticipantComponent.h
--- a/Engine/Plugins/Experimental/CommonConversation/Source/CommonConversationRuntime/Public/ConversationParticipantComponent.h
+++ b/Engine/Plugins/Experimental/CommonConversation/Source/CommonConversationRuntime/Public/ConversationParticipantComponent.h
@@ -40,7 +40,7 @@ public:
 
 public:
 	void SendClientConversationMessage(const FConversationContext& Context, const FClientConversationMessagePayload& Payload);
-	void SendClientUpdatedChoices(const FConversationContext& Context);
+	virtual void SendClientUpdatedChoices(const FConversationContext& Context);
 	void SendClientRefreshedTaskChoiceData(const FConversationNodeHandle& Handle, const FConversationContext& Context);
 
 	UFUNCTION(BlueprintCallable, Category=Conversation)
