Conversations - Potential fix to low frequency crash, the LastMessage data on the conversation participant component has a strong possibility that the Actor* in the participants list can be deleted, but wont be properly cleaned up (because it's not a UProperty) leading to garbage ptrs being serialized or accessed.

diff --git a/Engine/Plugins/Experimental/CommonConversation/Source/CommonConversationRuntime/Private/ConversationContext.cpp b/Engine/Plugins/Experimental/CommonConversation/Source/CommonConversationRuntime/Private/ConversationContext.cpp
--- a/Engine/Plugins/Experimental/CommonConversation/Source/CommonConversationRuntime/Private/ConversationContext.cpp
+++ b/Engine/Plugins/Experimental/CommonConversation/Source/CommonConversationRuntime/Private/ConversationContext.cpp
@@ -100,7 +100,7 @@ FConversationParticipants FConversationContext::GetParticipantsCopy() const
 	{
 		if (ensure(ClientParticipant))
 		{
-			return ClientParticipant->LastMessage.Participants;
+			return ClientParticipant->GetLastMessage().Participants;
 		}
 	}
 
diff --git a/Engine/Plugins/Experimental/CommonConversation/Source/CommonConversationRuntime/Public/ConversationParticipantComponent.h b/Engine/Plugins/Experimental/CommonConversation/Source/CommonConversationRuntime/Public/ConversationParticipantComponent.h
--- a/Engine/Plugins/Experimental/CommonConversation/Source/CommonConversationRuntime/Public/ConversationParticipantComponent.h
+++ b/Engine/Plugins/Experimental/CommonConversation/Source/CommonConversationRuntime/Public/ConversationParticipantComponent.h
@@ -72,7 +72,6 @@ public:
 	UFUNCTION(BlueprintCallable, Category=Conversation)
 	virtual FText GetParticipantDisplayName();
 
-
 	UFUNCTION(BlueprintCallable, Category = Conversation)
 	bool IsInActiveConversation() const;
 
@@ -123,10 +122,15 @@ protected:
 #endif
 
 public:
-	FClientConversationMessagePayload LastMessage;
-	int32 MessageIndex = 0;
-
+    // The number of conversations active.  A given conversationalist might be in multiple conversations at once.
+    // e.g. Multiple players "talking" to the same NPC in a multiplayer game.
 	int32 GetConversationsActive() const { return ConversationsActive; }
+	
+	// A cached version of the last conversation message payload data recieved.
+	const FClientConversationMessagePayload& GetLastMessage() const { return LastMessage; }
+	
+	// Gets the last message index recieved, this is just a monotomically increasing number each time we get a new message.
+	const int32 GetLastMessageIndex() const { return MessageIndex; }
 
 	bool GetIsFirstConversationUpdateBroadcasted() const { return bIsFirstConversationUpdateBroadcasted; }
 
@@ -146,7 +150,12 @@ private:
 
 	UPROPERTY()
 	TArray<UConversationInstance*> Auth_Conversations;
-
+	
+	UPROPERTY()
+    FClientConversationMessagePayload LastMessage;
+    
+    int32 MessageIndex = 0;
+    
 	bool bIsFirstConversationUpdateBroadcasted = false;
 };
 
