Conversation - Forward declaring the conversation node.

diff --git a/Engine/Plugins/Experimental/CommonConversation/Source/CommonConversationRuntime/Public/ConversationNode.h b/Engine/Plugins/Experimental/CommonConversation/Source/CommonConversationRuntime/Public/ConversationNode.h
--- a/Engine/Plugins/Experimental/CommonConversation/Source/CommonConversationRuntime/Public/ConversationNode.h
+++ b/Engine/Plugins/Experimental/CommonConversation/Source/CommonConversationRuntime/Public/ConversationNode.h
@@ -11,6 +11,7 @@
 
 class UWorld;
 
+class UConversationNode;
 class UConversationDatabase;
 class UCommonDialogueConversation;
 struct FConversationContext;
