Common Conversation - Fixing empty document tab being popped up.
Common Conversation - Making the StartConversation library call Authority only as is already fact, just communicating that better.

diff --git a/Engine/Plugins/Experimental/CommonConversation/Source/CommonConversationEditor/Private/ConversationEditor.cpp b/Engine/Plugins/Experimental/CommonConversation/Source/CommonConversationEditor/Private/ConversationEditor.cpp
--- a/Engine/Plugins/Experimental/CommonConversation/Source/CommonConversationEditor/Private/ConversationEditor.cpp
+++ b/Engine/Plugins/Experimental/CommonConversation/Source/CommonConversationEditor/Private/ConversationEditor.cpp
@@ -245,14 +245,14 @@ void FConversationEditor::RestoreConversation()
 
 			if (UConversationGraph* ConversationGraph = Cast<UConversationGraph>(DocumentObject))
 			{
-				TSharedPtr<SDockTab> DocumentTab = DocumentManager->OpenDocument(Payload, FDocumentTracker::OpenNewDocument);
+				TSharedPtr<SDockTab> DocumentTab = DocumentManager->OpenDocument(Payload, FDocumentTracker::RestorePreviousDocument);
 
 				TSharedRef<SGraphEditor> GraphEditor = StaticCastSharedRef<SGraphEditor>(DocumentTab->GetContent());
 				GraphEditor->SetViewLocation(DocumentInfo.SavedViewOffset, DocumentInfo.SavedZoomAmount);
 			}
 			else
 			{
-				DocumentManager->OpenDocument(Payload, FDocumentTracker::OpenNewDocument);
+				DocumentManager->OpenDocument(Payload, FDocumentTracker::RestorePreviousDocument);
 			}
 		}
 	}
@@ -355,16 +355,6 @@ TSharedRef<SGraphEditor> FConversationEditor::CreateGraphEditorWidget(UEdGraph*
 	{
 		CreateCommandList();
 
-// 		GraphEditorCommands->MapAction( FGraphEditorCommands::Get().RemoveExecutionPin,
-// 			FExecuteAction::CreateSP( this, &FConversationEditor::OnRemoveInputPin ),
-// 			FCanExecuteAction::CreateSP( this, &FConversationEditor::CanRemoveInputPin )
-// 			);
-// 
-// 		GraphEditorCommands->MapAction( FGraphEditorCommands::Get().AddExecutionPin,
-// 			FExecuteAction::CreateSP( this, &FConversationEditor::OnAddInputPin ),
-// 			FCanExecuteAction::CreateSP( this, &FConversationEditor::CanAddInputPin )
-// 			);
-
 		// Debug actions
 		GraphEditorCommands->MapAction( FGraphEditorCommands::Get().AddBreakpoint,
 			FExecuteAction::CreateSP( this, &FConversationEditor::OnAddBreakpoint ),
diff --git a/Engine/Plugins/Experimental/CommonConversation/Source/CommonConversationEditor/Private/ConversationEditorModes.cpp b/Engine/Plugins/Experimental/CommonConversation/Source/CommonConversationEditor/Private/ConversationEditorModes.cpp
--- a/Engine/Plugins/Experimental/CommonConversation/Source/CommonConversationEditor/Private/ConversationEditorModes.cpp
+++ b/Engine/Plugins/Experimental/CommonConversation/Source/CommonConversationEditor/Private/ConversationEditorModes.cpp
@@ -16,9 +16,8 @@ FConversationEditorApplicationMode_GraphView::FConversationEditorApplicationMode
 	ConversationEditorTabFactories.RegisterFactory(MakeShareable(new FConversationDetailsSummoner(InConversationEditor)));
 	ConversationEditorTabFactories.RegisterFactory(MakeShareable(new FConversationSearchSummoner(InConversationEditor)));
 	ConversationEditorTabFactories.RegisterFactory(MakeShareable(new FConversationTreeEditorSummoner(InConversationEditor)));
-	//ConversationEditorTabFactories.RegisterFactory(MakeShareable(new FBlackboardSummoner(InConversationEditor)));
 
-	TabLayout = FTabManager::NewLayout( "Standalone_ConversationEditor_GraphView_Layout_v2" )
+	TabLayout = FTabManager::NewLayout( "Standalone_ConversationEditor_GraphView_Layout_v4" )
 	->AddArea
 	(
 		FTabManager::NewPrimaryArea()
diff --git a/Engine/Plugins/Experimental/CommonConversation/Source/CommonConversationEditor/Private/ConversationEditorTabs.cpp b/Engine/Plugins/Experimental/CommonConversation/Source/CommonConversationEditor/Private/ConversationEditorTabs.cpp
--- a/Engine/Plugins/Experimental/CommonConversation/Source/CommonConversationEditor/Private/ConversationEditorTabs.cpp
+++ b/Engine/Plugins/Experimental/CommonConversation/Source/CommonConversationEditor/Private/ConversationEditorTabs.cpp
@@ -8,4 +8,4 @@ const FName FConversationEditorTabs::SearchID(TEXT("ConversationEditor_Search"))
 const FName FConversationEditorTabs::TreeEditorID(TEXT("ConversationEditor_Tree"));
 
 // Document tab identifiers
-const FName FConversationEditorTabs::GraphEditorID(TEXT("GraphDocument"));
+const FName FConversationEditorTabs::GraphEditorID(TEXT("Document"));
diff --git a/Engine/Plugins/Experimental/CommonConversation/Source/CommonConversationRuntime/Private/ConversationDatabase.cpp b/Engine/Plugins/Experimental/CommonConversation/Source/CommonConversationRuntime/Private/ConversationDatabase.cpp
--- a/Engine/Plugins/Experimental/CommonConversation/Source/CommonConversationRuntime/Private/ConversationDatabase.cpp
+++ b/Engine/Plugins/Experimental/CommonConversation/Source/CommonConversationRuntime/Private/ConversationDatabase.cpp
@@ -80,38 +80,41 @@ EDataValidationResult UConversationDatabase::IsDataValid(TArray<FText>& Validati
 
 	EDataValidationResult Result = EDataValidationResult::Valid;
 
-	const TSet<FName>* CookedAssetTagsForConversations = SaveOptions.CookFilterlistTagsByClass.Find(UConversationDatabase::StaticClass()->GetClassPathName());
-	if (CookedAssetTagsForConversations == nullptr || !CookedAssetTagsForConversations->Contains(GET_MEMBER_NAME_CHECKED(UConversationDatabase, EntryTags)))
+	if (SaveOptions.bUseAssetRegistryTagsAllowListInsteadOfDenyList)
 	{
-		ValidationErrors.Add(FText::Format(LOCTEXT("Missing_EntryTags", "Missing from DefaultEngine.ini, {0}"),
-			FText::FromString(TEXT("+CookedTagsWhitelist=(Class=ConversationDatabase,Tag=EntryTags)")))
-		);
+		const TSet<FName>* CookedAssetTagsForConversations = SaveOptions.CookFilterlistTagsByClass.Find(UConversationDatabase::StaticClass()->GetClassPathName());
+		if (CookedAssetTagsForConversations == nullptr || !CookedAssetTagsForConversations->Contains(GET_MEMBER_NAME_CHECKED(UConversationDatabase, EntryTags)))
+		{
+			ValidationErrors.Add(FText::Format(LOCTEXT("Missing_EntryTags", "Missing from DefaultEngine.ini, {0}"),
+				FText::FromString(TEXT("+CookedTagsWhitelist=(Class=ConversationDatabase,Tag=EntryTags)")))
+			);
 
-		Result = EDataValidationResult::Invalid;
-	}
-	if (CookedAssetTagsForConversations == nullptr || !CookedAssetTagsForConversations->Contains(GET_MEMBER_NAME_CHECKED(UConversationDatabase, ExitTags)))
-	{
-		ValidationErrors.Add(FText::Format(LOCTEXT("Missing_ExitTags", "Missing from DefaultEngine.ini, {0}"),
-			FText::FromString(TEXT("+CookedTagsWhitelist=(Class=ConversationDatabase,Tag=ExitTags)")))
-		);
+			Result = EDataValidationResult::Invalid;
+		}
+		if (CookedAssetTagsForConversations == nullptr || !CookedAssetTagsForConversations->Contains(GET_MEMBER_NAME_CHECKED(UConversationDatabase, ExitTags)))
+		{
+			ValidationErrors.Add(FText::Format(LOCTEXT("Missing_ExitTags", "Missing from DefaultEngine.ini, {0}"),
+				FText::FromString(TEXT("+CookedTagsWhitelist=(Class=ConversationDatabase,Tag=ExitTags)")))
+			);
 
-		Result = EDataValidationResult::Invalid;
-	}
-	if (CookedAssetTagsForConversations == nullptr || !CookedAssetTagsForConversations->Contains(GET_MEMBER_NAME_CHECKED(UConversationDatabase, InternalNodeIds)))
-	{
-		ValidationErrors.Add(FText::Format(LOCTEXT("Missing_InternalNodeIds", "Missing from DefaultEngine.ini, {0}"),
-			FText::FromString(TEXT("+CookedTagsWhitelist=(Class=ConversationDatabase,Tag=InternalNodeIds)")))
-		);
+			Result = EDataValidationResult::Invalid;
+		}
+		if (CookedAssetTagsForConversations == nullptr || !CookedAssetTagsForConversations->Contains(GET_MEMBER_NAME_CHECKED(UConversationDatabase, InternalNodeIds)))
+		{
+			ValidationErrors.Add(FText::Format(LOCTEXT("Missing_InternalNodeIds", "Missing from DefaultEngine.ini, {0}"),
+				FText::FromString(TEXT("+CookedTagsWhitelist=(Class=ConversationDatabase,Tag=InternalNodeIds)")))
+			);
 
-		Result = EDataValidationResult::Invalid;
-	}
-	if (CookedAssetTagsForConversations == nullptr || !CookedAssetTagsForConversations->Contains(GET_MEMBER_NAME_CHECKED(UConversationDatabase, LinkedToNodeIds)))
-	{
-		ValidationErrors.Add(FText::Format(LOCTEXT("Missing_LinkedToNodeIds", "Missing from DefaultEngine.ini, {0}"),
-			FText::FromString(TEXT("+CookedTagsWhitelist=(Class=ConversationDatabase,Tag=LinkedToNodeIds)")))
-		);
+			Result = EDataValidationResult::Invalid;
+		}
+		if (CookedAssetTagsForConversations == nullptr || !CookedAssetTagsForConversations->Contains(GET_MEMBER_NAME_CHECKED(UConversationDatabase, LinkedToNodeIds)))
+		{
+			ValidationErrors.Add(FText::Format(LOCTEXT("Missing_LinkedToNodeIds", "Missing from DefaultEngine.ini, {0}"),
+				FText::FromString(TEXT("+CookedTagsWhitelist=(Class=ConversationDatabase,Tag=LinkedToNodeIds)")))
+			);
 
-		Result = EDataValidationResult::Invalid;
+			Result = EDataValidationResult::Invalid;
+		}
 	}
 
 	return CombineDataValidationResults(SuperResult, Result);
diff --git a/Engine/Plugins/Experimental/CommonConversation/Source/CommonConversationRuntime/Public/ConversationLibrary.h b/Engine/Plugins/Experimental/CommonConversation/Source/CommonConversationRuntime/Public/ConversationLibrary.h
--- a/Engine/Plugins/Experimental/CommonConversation/Source/CommonConversationRuntime/Public/ConversationLibrary.h
+++ b/Engine/Plugins/Experimental/CommonConversation/Source/CommonConversationRuntime/Public/ConversationLibrary.h
@@ -21,6 +21,6 @@ class COMMONCONVERSATIONRUNTIME_API UConversationLibrary : public UBlueprintFunc
 public:
 	UConversationLibrary();
 
-	UFUNCTION(BlueprintCallable, Category="Conversation", meta = (DeterminesOutputType = "ConversationType"))
+	UFUNCTION(BlueprintCallable, BlueprintAuthorityOnly, Category="Conversation", meta = (DeterminesOutputType = "ConversationType"))
 	static UConversationInstance* StartConversation(FGameplayTag ConversationEntryTag, AActor* Instigator, FGameplayTag InstigatorTag, AActor* Target, FGameplayTag TargetTag);
 };
