Common Conversation - Swapping the world world context before calling the side effect nodes, otherwise world context might be wrong or null.

diff --git a/Engine/Plugins/Experimental/CommonConversation/Source/CommonConversationRuntime/Private/ConversationSideEffectNode.cpp b/Engine/Plugins/Experimental/CommonConversation/Source/CommonConversationRuntime/Private/ConversationSideEffectNode.cpp
--- a/Engine/Plugins/Experimental/CommonConversation/Source/CommonConversationRuntime/Private/ConversationSideEffectNode.cpp
+++ b/Engine/Plugins/Experimental/CommonConversation/Source/CommonConversationRuntime/Private/ConversationSideEffectNode.cpp
@@ -2,9 +2,12 @@
 
 #include "ConversationSideEffectNode.h"
 #include "ConversationInstance.h"
+#include "Engine/World.h"
 
 void UConversationSideEffectNode::CauseSideEffect(const FConversationContext& Context) const
 {
+	TGuardValue<decltype(EvalWorldContextObj)> Swapper(EvalWorldContextObj, Context.GetWorld());
+	
 	if (Context.IsServerContext())
 	{
 		ServerCauseSideEffect(Context);
