Add handling for invalid branch choice in conversations
Creates a virtual function to be overridden to determine how invalid choices should be handled in conversation instances

diff --git a/Engine/Plugins/Experimental/CommonConversation/Source/CommonConversationRuntime/Private/ConversationInstance.cpp b/Engine/Plugins/Experimental/CommonConversation/Source/CommonConversationRuntime/Private/ConversationInstance.cpp
--- a/Engine/Plugins/Experimental/CommonConversation/Source/CommonConversationRuntime/Private/ConversationInstance.cpp
+++ b/Engine/Plugins/Experimental/CommonConversation/Source/CommonConversationRuntime/Private/ConversationInstance.cpp
@@ -215,8 +215,7 @@ void UConversationInstance::ServerAdvanceConversation(const FAdvanceConversation
 			}
 			else
 			{
-				UE_LOG(LogCommonConversationRuntime, Error, TEXT("User picked option %s but it's not a legal output, aborting"), *InChoicePicked.ToString());
-				ServerAbortConversation();
+				OnInvalidBranchChoice(InChoicePicked);
 				return;
 			}
 		}
@@ -282,6 +281,12 @@ void UConversationInstance::ServerAdvanceConversation(const FAdvanceConversation
 	}
 }
 
+void UConversationInstance::OnInvalidBranchChoice(const FAdvanceConversationRequest& InChoicePicked)
+{
+	UE_LOG(LogCommonConversationRuntime, Error, TEXT("User picked option %s but it's not a legal output, aborting"), *InChoicePicked.ToString());
+	ServerAbortConversation();
+}
+
 void UConversationInstance::ServerAbortConversation()
 {
 	if (bConversationStarted)
diff --git a/Engine/Plugins/Experimental/CommonConversation/Source/CommonConversationRuntime/Private/ConversationLibrary.cpp b/Engine/Plugins/Experimental/CommonConversation/Source/CommonConversationRuntime/Private/ConversationLibrary.cpp
--- a/Engine/Plugins/Experimental/CommonConversation/Source/CommonConversationRuntime/Private/ConversationLibrary.cpp
+++ b/Engine/Plugins/Experimental/CommonConversation/Source/CommonConversationRuntime/Private/ConversationLibrary.cpp
@@ -15,7 +15,8 @@ UConversationLibrary::UConversationLibrary()
 {
 }
 
-UConversationInstance* UConversationLibrary::StartConversation(FGameplayTag ConversationEntryTag, AActor* Instigator, FGameplayTag InstigatorTag, AActor* Target, FGameplayTag TargetTag)
+UConversationInstance* UConversationLibrary::StartConversation(FGameplayTag ConversationEntryTag, AActor* Instigator,
+	FGameplayTag InstigatorTag, AActor* Target, FGameplayTag TargetTag, const TSubclassOf<UConversationInstance> ConversationInstanceClass)
 {
 #if WITH_SERVER_CODE
 	if (Instigator == nullptr || Target == nullptr)
@@ -25,12 +26,15 @@ UConversationInstance* UConversationLibrary::StartConversation(FGameplayTag Conv
 
 	if (UWorld* World = GEngine->GetWorldFromContextObject(Instigator, EGetWorldErrorMode::LogAndReturnNull))
 	{
-		UClass* InstanceClass = GetDefault<UConversationSettings>()->GetConversationInstanceClass();
-		if (InstanceClass == nullptr)
+		UClass* InstanceClass = ConversationInstanceClass;
+		if (!InstanceClass)
 		{
-			InstanceClass = UConversationInstance::StaticClass();
+			InstanceClass = GetDefault<UConversationSettings>()->GetConversationInstanceClass();
+			if (InstanceClass == nullptr)
+			{
+				InstanceClass = UConversationInstance::StaticClass();
+			}
 		}
-
 		UConversationInstance* ConversationInstance = NewObject<UConversationInstance>(World, InstanceClass);
 		if (ensure(ConversationInstance))
 		{
diff --git a/Engine/Plugins/Experimental/CommonConversation/Source/CommonConversationRuntime/Public/ConversationInstance.h b/Engine/Plugins/Experimental/CommonConversation/Source/CommonConversationRuntime/Public/ConversationInstance.h
--- a/Engine/Plugins/Experimental/CommonConversation/Source/CommonConversationRuntime/Public/ConversationInstance.h
+++ b/Engine/Plugins/Experimental/CommonConversation/Source/CommonConversationRuntime/Public/ConversationInstance.h
@@ -47,6 +47,8 @@ public:
 
 	void ServerAdvanceConversation(const FAdvanceConversationRequest& InChoicePicked);
 
+	virtual void OnInvalidBranchChoice(const FAdvanceConversationRequest& InChoicePicked);
+
 	void ServerAbortConversation();
 
 	void ServerRefreshConversationChoices();
diff --git a/Engine/Plugins/Experimental/CommonConversation/Source/CommonConversationRuntime/Public/ConversationLibrary.h b/Engine/Plugins/Experimental/CommonConversation/Source/CommonConversationRuntime/Public/ConversationLibrary.h
--- a/Engine/Plugins/Experimental/CommonConversation/Source/CommonConversationRuntime/Public/ConversationLibrary.h
+++ b/Engine/Plugins/Experimental/CommonConversation/Source/CommonConversationRuntime/Public/ConversationLibrary.h
@@ -22,5 +22,6 @@ public:
 	UConversationLibrary();
 
 	UFUNCTION(BlueprintCallable, BlueprintAuthorityOnly, Category="Conversation", meta = (DeterminesOutputType = "ConversationType"))
-	static UConversationInstance* StartConversation(FGameplayTag ConversationEntryTag, AActor* Instigator, FGameplayTag InstigatorTag, AActor* Target, FGameplayTag TargetTag);
+	static UConversationInstance* StartConversation(FGameplayTag ConversationEntryTag, AActor* Instigator, FGameplayTag InstigatorTag,
+		AActor* Target, FGameplayTag TargetTag, const TSubclassOf<UConversationInstance> ConversationInstanceClass = nullptr);
 };
