Fixed issue with PrimaryAssetId reporting where the joined asset id string appeared as junk in the log. Changed the JoinBy projection lambda to return a FString rather than TCHAR*

diff --git a/Engine/Plugins/Experimental/CommonConversation/Source/CommonConversationRuntime/Private/ConversationRegistry.cpp b/Engine/Plugins/Experimental/CommonConversation/Source/CommonConversationRuntime/Private/ConversationRegistry.cpp
--- a/Engine/Plugins/Experimental/CommonConversation/Source/CommonConversationRuntime/Private/ConversationRegistry.cpp
+++ b/Engine/Plugins/Experimental/CommonConversation/Source/CommonConversationRuntime/Private/ConversationRegistry.cpp
@@ -355,7 +355,7 @@ TArray<FPrimaryAssetId> UConversationRegistry::GetPrimaryAssetIdsForEntryPoint(F
 
 	UE_LOG(LogCommonConversationRuntime, Verbose, TEXT("For EntryPoint [%s], Found [%s]"),
 		*EntryPoint.ToString(),
-		*FString::JoinBy(AssetsWithTheEntryPoint, TEXT(", "), [](const FPrimaryAssetId& PrimaryAssetId) { return *PrimaryAssetId.ToString(); })
+		*FString::JoinBy(AssetsWithTheEntryPoint, TEXT(", "), [](const FPrimaryAssetId& PrimaryAssetId) { return PrimaryAssetId.ToString(); })
 	);
 
 	return AssetsWithTheEntryPoint;
