Added stat scope for UConversationRegistry::BuildDependenciesGraph.

diff --git a/Engine/Plugins/Experimental/CommonConversation/Source/CommonConversationRuntime/Private/ConversationRegistry.cpp b/Engine/Plugins/Experimental/CommonConversation/Source/CommonConversationRuntime/Private/ConversationRegistry.cpp
--- a/Engine/Plugins/Experimental/CommonConversation/Source/CommonConversationRuntime/Private/ConversationRegistry.cpp
+++ b/Engine/Plugins/Experimental/CommonConversation/Source/CommonConversationRuntime/Private/ConversationRegistry.cpp
@@ -13,6 +13,7 @@
 #include "ConversationContext.h"
 #include "Engine/StreamableManager.h"
 #include "GameFeaturesSubsystem.h"
+#include "Stats/Stats2.h"
 
 #include UE_INLINE_GENERATED_CPP_BY_NAME(ConversationRegistry)
 
@@ -371,6 +372,7 @@ void UConversationRegistry::BuildDependenciesGraph()
 		return;
 	}
 
+	QUICK_SCOPE_CYCLE_COUNTER(STAT_UConversationRegistry_BuildDependenciesGraph);
 	UE_LOG(LogCommonConversationRuntime, Verbose, TEXT("Registry Building Graph"));
 
 	TArray<FAssetData> AllActiveConversationAssets;
