CommonConversation - Fixing a few listening server bugs in the common conversation layer.  Needed to store different flags for client/server since client is not just !server in the case of a listen server.

diff --git a/Engine/Plugins/Experimental/CommonConversation/Source/CommonConversationRuntime/Private/ConversationContext.cpp b/Engine/Plugins/Experimental/CommonConversation/Source/CommonConversationRuntime/Private/ConversationContext.cpp
--- a/Engine/Plugins/Experimental/CommonConversation/Source/CommonConversationRuntime/Private/ConversationContext.cpp
+++ b/Engine/Plugins/Experimental/CommonConversation/Source/CommonConversationRuntime/Private/ConversationContext.cpp
@@ -11,22 +11,28 @@
 
 FConversationContext FConversationContext::CreateServerContext(UConversationInstance* InActiveConversation, const UConversationTaskNode* InTaskBeingConsidered)
 {
+	UWorld* World = InActiveConversation->GetWorld();
+
 	FConversationContext Context;
 	Context.ActiveConversation = InActiveConversation;
 	Context.TaskBeingConsidered = InTaskBeingConsidered;
-	Context.bServer = true;
-	Context.ConversationRegistry = UConversationRegistry::GetFromWorld(InActiveConversation->GetWorld());
+	Context.bServer_PRIVATE = true;
+	Context.bClient_PRIVATE = World->GetNetMode() != NM_DedicatedServer;
+	Context.ConversationRegistry = UConversationRegistry::GetFromWorld(World);
 
 	return Context;
 }
 
 FConversationContext FConversationContext::CreateClientContext(UConversationParticipantComponent* InParticipantComponent, const UConversationTaskNode* InTaskBeingConsidered)
 {
+	UWorld* World = InParticipantComponent->GetWorld();
+
 	FConversationContext Context;
 	Context.ClientParticipant = InParticipantComponent;
 	Context.TaskBeingConsidered = InTaskBeingConsidered;
-	Context.bServer = false;
-	Context.ConversationRegistry = UConversationRegistry::GetFromWorld(InParticipantComponent->GetWorld());
+	Context.bServer_PRIVATE = false;
+	Context.bClient_PRIVATE = true;
+	Context.ConversationRegistry = UConversationRegistry::GetFromWorld(World);
 
 	return Context;
 }
@@ -49,7 +55,7 @@ FConversationContext FConversationContext::CreateReturnScopeContext(const FConve
 
 UWorld* FConversationContext::GetWorld() const
 {
-	if (bServer)
+	if (IsServerContext())
 	{
 		if (ActiveConversation)
 		{
@@ -69,7 +75,7 @@ UWorld* FConversationContext::GetWorld() const
 
 FConversationNodeHandle FConversationContext::GetCurrentNodeHandle() const
 {
-	if (bServer)
+	if (IsServerContext())
 	{
 		if (ActiveConversation)
 		{
@@ -89,7 +95,7 @@ FConversationNodeHandle FConversationContext::GetCurrentNodeHandle() const
 
 FConversationParticipants FConversationContext::GetParticipantsCopy() const
 {
-	if (bServer)
+	if (IsServerContext())
 	{
 		if (ActiveConversation)
 		{
@@ -109,7 +115,7 @@ FConversationParticipants FConversationContext::GetParticipantsCopy() const
 
 const FConversationParticipantEntry* FConversationContext::GetParticipant(const FGameplayTag& ParticipantTag) const
 {
-	if (bServer)
+	if (IsServerContext())
 	{
 		if (ActiveConversation)
 		{
diff --git a/Engine/Plugins/Experimental/CommonConversation/Source/CommonConversationRuntime/Private/ConversationParticipantComponent.cpp b/Engine/Plugins/Experimental/CommonConversation/Source/CommonConversationRuntime/Private/ConversationParticipantComponent.cpp
--- a/Engine/Plugins/Experimental/CommonConversation/Source/CommonConversationRuntime/Private/ConversationParticipantComponent.cpp
+++ b/Engine/Plugins/Experimental/CommonConversation/Source/CommonConversationRuntime/Private/ConversationParticipantComponent.cpp
@@ -43,9 +43,11 @@ void UConversationParticipantComponent::ServerNotifyConversationStarted(UConvers
 		// executed, the client has knowledge of what participants, before any client side task effects need to execute.
 		ClientUpdateParticipants(Auth_CurrentConversation->GetParticipantsCopy());
 
-		if (ConversationsActive == 0)
+		const int32 OldConversationsActive = ConversationsActive;
+
+		if (OldConversationsActive == 0)
 		{
-			OnEnterConversationState();
+			OnRep_ConversationsActive(OldConversationsActive);
 		}
 
 		ConversationsActive++;
@@ -54,10 +56,7 @@ void UConversationParticipantComponent::ServerNotifyConversationStarted(UConvers
 		OnServerConversationStarted(Conversation, AsParticipant);
 		ClientStartConversation(AsParticipant);
 
-		if (Owner->GetRemoteRole() == ROLE_AutonomousProxy)
-		{
-			ClientUpdateConversations(ConversationsActive);
-		}
+		ClientUpdateConversations(ConversationsActive);
 	}
 }
 
@@ -72,19 +71,17 @@ void UConversationParticipantComponent::ServerNotifyConversationEnded(UConversat
 			Auth_CurrentConversation = nullptr;
 			Auth_Conversations.Remove(Conversation);
 
+			const int32 OldConversationsActive = ConversationsActive;
 			ConversationsActive--;
 			MARK_PROPERTY_DIRTY_FROM_NAME(UConversationParticipantComponent, ConversationsActive, this);
 
 			OnServerConversationEnded(Conversation);
 
-			if (Owner->GetRemoteRole() == ROLE_AutonomousProxy)
-			{
-				ClientUpdateConversations(ConversationsActive);
-			}
+			ClientUpdateConversations(ConversationsActive);
 
 			if (ConversationsActive == 0)
 			{
-				OnLeaveConversationState();
+				OnRep_ConversationsActive(OldConversationsActive);
 			}
 		}
 	}
@@ -169,10 +166,7 @@ void UConversationParticipantComponent::SendClientConversationMessage(const FCon
 	LastMessage = Payload;
 
 	//@TODO: CONVERSATION: We could potentially send the user no choices?  I guess that's a possibility.
-	if (GetOwner()->GetRemoteRole() == ROLE_AutonomousProxy)
-	{
-		ClientUpdateConversation(LastMessage);
-	}
+	ClientUpdateConversation(LastMessage);
 #endif
 }
 
@@ -186,14 +180,11 @@ void UConversationParticipantComponent::ClientUpdateConversations_Implementation
 void UConversationParticipantComponent::SendClientUpdatedChoices(const FConversationContext& Context)
 {
 #if WITH_SERVER_CODE
-	if (GetOwner()->GetRemoteRole() == ROLE_AutonomousProxy)
+	const TArray<FClientConversationOptionEntry> NewOptions = Context.GetActiveConversation()->GetCurrentUserConversationChoices();
+	if (NewOptions != LastMessage.Options)
 	{
-		const TArray<FClientConversationOptionEntry> NewOptions = Context.GetActiveConversation()->GetCurrentUserConversationChoices();
-		if (NewOptions != LastMessage.Options)
-		{
-			LastMessage.Options = NewOptions;
-			ClientUpdateConversation(LastMessage);
-		}
+		LastMessage.Options = NewOptions;
+		ClientUpdateConversation(LastMessage);
 	}
 #endif
 }
@@ -201,23 +192,20 @@ void UConversationParticipantComponent::SendClientUpdatedChoices(const FConversa
 void UConversationParticipantComponent::SendClientRefreshedTaskChoiceData(const FConversationNodeHandle& Handle, const FConversationContext& Context)
 {
 #if WITH_SERVER_CODE
-	if (GetOwner()->GetRemoteRole() == ROLE_AutonomousProxy)
+	const TArray<FClientConversationOptionEntry> CurrentOptions = Context.GetActiveConversation()->GetCurrentUserConversationChoices();
+	
+	for (const FClientConversationOptionEntry& CurrentOption : CurrentOptions)
 	{
-		const TArray<FClientConversationOptionEntry> CurrentOptions = Context.GetActiveConversation()->GetCurrentUserConversationChoices();
-		
-		for (const FClientConversationOptionEntry& CurrentOption : CurrentOptions)
+		if (CurrentOption.ChoiceReference.NodeReference == Handle)
 		{
-			if (CurrentOption.ChoiceReference.NodeReference == Handle)
+			for (FClientConversationOptionEntry& LastOption : LastMessage.Options)
 			{
-				for (FClientConversationOptionEntry& LastOption : LastMessage.Options)
+				if (LastOption.ChoiceReference.NodeReference == Handle && 
+					LastOption.ExtraData != CurrentOption.ExtraData)
 				{
-					if (LastOption.ChoiceReference.NodeReference == Handle && 
-						LastOption.ExtraData != CurrentOption.ExtraData)
-					{
-						LastOption.ExtraData = CurrentOption.ExtraData;
-						ClientUpdateConversationTaskChoiceData(Handle, LastOption);
-						break;
-					}
+					LastOption.ExtraData = CurrentOption.ExtraData;
+					ClientUpdateConversationTaskChoiceData(Handle, LastOption);
+					break;
 				}
 			}
 		}
@@ -307,12 +295,15 @@ void UConversationParticipantComponent::ClientExecuteTaskAndSideEffects_Implemen
 
 void UConversationParticipantComponent::OnRep_ConversationsActive(int32 OldConversationsActive)
 {
-	if (OldConversationsActive == 0)
+	const bool bWasConversing = OldConversationsActive > 0;
+	const bool bIsConversing = ConversationsActive > 0;
+
+	if (!bWasConversing && bIsConversing)
 	{
 		OnEnterConversationState();
 		ConversationStatusChanged.Broadcast(true);
 	}
-	else if (ConversationsActive == 0)
+	else if (bWasConversing && !bIsConversing)
 	{
 		OnLeaveConversationState();
 		ConversationStatusChanged.Broadcast(false);
@@ -353,7 +344,7 @@ void UConversationParticipantComponent::OnConversationUpdated(const FClientConve
 void UConversationParticipantComponent::ServerAbortAllConversations()
 {
 	TArray<UConversationInstance*> Conversations = Auth_Conversations;
-	for(UConversationInstance* Conversation : Conversations)
+	for (UConversationInstance* Conversation : Conversations)
 	{
 		Conversation->ServerAbortConversation();
 	}
diff --git a/Engine/Plugins/Experimental/CommonConversation/Source/CommonConversationRuntime/Private/ConversationTaskNode.cpp b/Engine/Plugins/Experimental/CommonConversation/Source/CommonConversationRuntime/Private/ConversationTaskNode.cpp
--- a/Engine/Plugins/Experimental/CommonConversation/Source/CommonConversationRuntime/Private/ConversationTaskNode.cpp
+++ b/Engine/Plugins/Experimental/CommonConversation/Source/CommonConversationRuntime/Private/ConversationTaskNode.cpp
@@ -131,10 +131,7 @@ FConversationTaskResult UConversationTaskNode::ExecuteTaskNodeWithSideEffects(co
 				if (UConversationParticipantComponent* Component = ParticipantEntry.GetParticipantComponent())
 				{
 					// Notify each client in the conversation
-					if (Component->GetOwner()->GetRemoteRole() == ROLE_AutonomousProxy)
-					{
-						Component->ServerNotifyExecuteTaskAndSideEffects(InContext.GetCurrentNodeHandle());
-					}
+					Component->ServerNotifyExecuteTaskAndSideEffects(InContext.GetCurrentNodeHandle());
 				}
 			}
 		}
diff --git a/Engine/Plugins/Experimental/CommonConversation/Source/CommonConversationRuntime/Public/ConversationContext.h b/Engine/Plugins/Experimental/CommonConversation/Source/CommonConversationRuntime/Public/ConversationContext.h
--- a/Engine/Plugins/Experimental/CommonConversation/Source/CommonConversationRuntime/Public/ConversationContext.h
+++ b/Engine/Plugins/Experimental/CommonConversation/Source/CommonConversationRuntime/Public/ConversationContext.h
@@ -177,7 +177,7 @@ public:
 
 	UConversationRegistry& GetConversationRegistry() const { check(ConversationRegistry) return *ConversationRegistry; }
 
-	UConversationInstance* GetActiveConversation() const { ensure(bServer); return ActiveConversation; }
+	UConversationInstance* GetActiveConversation() const { ensure(IsServerContext()); return ActiveConversation; }
 
 	const UConversationTaskNode* GetTaskBeingConsidered() const { return TaskBeingConsidered; }
 
@@ -207,8 +207,8 @@ public:
 
 	FConversationParticipants GetParticipantsCopy() const;
 
-	bool IsServerContext() const { return bServer; }
-	bool IsClientContext() const { return !bServer; }
+	bool IsServerContext() const { return bServer_PRIVATE; }
+	bool IsClientContext() const { return bClient_PRIVATE; }
 
 private:
 	UPROPERTY()
@@ -227,7 +227,10 @@ private:
 	TArray<FConversationNodeHandle> ReturnScopeStack;
 
 	UPROPERTY()
-	bool bServer = false;
+	bool bServer_PRIVATE = false;
+
+	UPROPERTY()
+	bool bClient_PRIVATE = false;
 };
 
 //////////////////////////////////////////////////////////////////////
