ConversationParticipantComponent: Added a new helper function that gets a reference to the in-conversation Actor that corresponds with that role parameter passed in.

diff --git a/Engine/Plugins/Experimental/CommonConversation/Source/CommonConversationRuntime/Private/ConversationParticipantComponent.cpp b/Engine/Plugins/Experimental/CommonConversation/Source/CommonConversationRuntime/Private/ConversationParticipantComponent.cpp
--- a/Engine/Plugins/Experimental/CommonConversation/Source/CommonConversationRuntime/Private/ConversationParticipantComponent.cpp
+++ b/Engine/Plugins/Experimental/CommonConversation/Source/CommonConversationRuntime/Private/ConversationParticipantComponent.cpp
@@ -160,6 +160,16 @@ const FConversationParticipantEntry* UConversationParticipantComponent::GetParti
 	return nullptr;
 }
 
+AActor* UConversationParticipantComponent::GetParticipantActor(const FGameplayTag& ParticipantTag) const
+{
+	if (const FConversationParticipantEntry* const ParticipantEntry = GetParticipant(ParticipantTag))
+	{
+		return ParticipantEntry->Actor;
+	}
+
+	return nullptr;
+}
+
 void UConversationParticipantComponent::SendClientConversationMessage(const FConversationContext& Context, const FClientConversationMessagePayload& Payload)
 {
 #if WITH_SERVER_CODE
diff --git a/Engine/Plugins/Experimental/CommonConversation/Source/CommonConversationRuntime/Public/ConversationParticipantComponent.h b/Engine/Plugins/Experimental/CommonConversation/Source/CommonConversationRuntime/Public/ConversationParticipantComponent.h
--- a/Engine/Plugins/Experimental/CommonConversation/Source/CommonConversationRuntime/Public/ConversationParticipantComponent.h
+++ b/Engine/Plugins/Experimental/CommonConversation/Source/CommonConversationRuntime/Public/ConversationParticipantComponent.h
@@ -80,6 +80,9 @@ public:
 
 	const FConversationParticipantEntry* GetParticipant(const FGameplayTag& ParticipantTag) const;
 
+	UFUNCTION(BlueprintPure, Category = Conversation)
+	AActor* GetParticipantActor(const FGameplayTag& ParticipantTag) const;
+
 protected:
 	UFUNCTION(Server, Reliable)
 	void ServerAdvanceConversation(const FAdvanceConversationRequest& InChoicePicked);
