CommonConverationRuntime Changes:
- Added a new virtual function to the Conversation Task Node base class: GenerateNextChoices. The default behavior is identical to existing choice generation behavior, but the virtual nature of the function allows for custom choice generation - like to support consecutive messaging.

#rnx

diff --git a/Engine/Plugins/Experimental/CommonConversation/Source/CommonConversationRuntime/Private/ConversationInstance.cpp b/Engine/Plugins/Experimental/CommonConversation/Source/CommonConversationRuntime/Private/ConversationInstance.cpp
--- a/Engine/Plugins/Experimental/CommonConversation/Source/CommonConversationRuntime/Private/ConversationInstance.cpp
+++ b/Engine/Plugins/Experimental/CommonConversation/Source/CommonConversationRuntime/Private/ConversationInstance.cpp
@@ -435,11 +435,10 @@ void UConversationInstance::UpdateNextChoices(const FConversationContext& Contex
 	if (const UConversationTaskNode* TaskNode = Cast<UConversationTaskNode>(GetCurrentChoiceReference().NodeReference.TryToResolve(Context)))
 	{
 		FConversationContext ChoiceContext = Context.CreateChildContext(TaskNode);
-
-		const TArray<FGuid> CandidateDestinations = ChoiceContext.GetConversationRegistry().GetOutputLinkGUIDs({ GetCurrentChoiceReference().NodeReference.NodeGUID });
-		
 		FConversationBranchPointBuilder BranchBuilder;
-		UConversationTaskNode::GenerateChoicesForDestinations(BranchBuilder, ChoiceContext, CandidateDestinations);
+
+		TaskNode->GenerateNextChoices(BranchBuilder, ChoiceContext);
+
 		AllChoices = BranchBuilder.GetBranches();
 	}
 
diff --git a/Engine/Plugins/Experimental/CommonConversation/Source/CommonConversationRuntime/Private/ConversationTaskNode.cpp b/Engine/Plugins/Experimental/CommonConversation/Source/CommonConversationRuntime/Private/ConversationTaskNode.cpp
--- a/Engine/Plugins/Experimental/CommonConversation/Source/CommonConversationRuntime/Private/ConversationTaskNode.cpp
+++ b/Engine/Plugins/Experimental/CommonConversation/Source/CommonConversationRuntime/Private/ConversationTaskNode.cpp
@@ -1,10 +1,12 @@
 // Copyright Epic Games, Inc. All Rights Reserved.
 
 #include "ConversationTaskNode.h"
+
 #include "CommonConversationRuntimeLogging.h"
+#include "ConversationChoiceNode.h"
 #include "ConversationContext.h"
+#include "ConversationInstance.h"
 #include "ConversationRegistry.h"
-#include "ConversationChoiceNode.h"
 #include "ConversationRequirementNode.h"
 #include "ConversationSideEffectNode.h"
 #include "Engine/World.h"
@@ -191,6 +193,16 @@ void UConversationTaskNode::GenerateChoicesForDestinations(FConversationBranchPo
 	}
 }
 
+void UConversationTaskNode::GenerateNextChoices(FConversationBranchPointBuilder& BranchBuilder, const FConversationContext& Context) const
+{
+	if (const UConversationInstance* Conversation = Context.GetActiveConversation())
+	{
+		const TArray<FGuid> CandidateDestinations = Context.GetConversationRegistry().GetOutputLinkGUIDs({ Conversation->GetCurrentChoiceReference().NodeReference.NodeGUID });
+
+		UConversationTaskNode::GenerateChoicesForDestinations(BranchBuilder, Context, CandidateDestinations);
+	}
+}
+
 void UConversationTaskNode::GatherChoices(FConversationBranchPointBuilder& BranchBuilder, const FConversationContext& Context) const
 {
 	GatherStaticChoices(BranchBuilder, Context);
diff --git a/Engine/Plugins/Experimental/CommonConversation/Source/CommonConversationRuntime/Public/ConversationTaskNode.h b/Engine/Plugins/Experimental/CommonConversation/Source/CommonConversationRuntime/Public/ConversationTaskNode.h
--- a/Engine/Plugins/Experimental/CommonConversation/Source/CommonConversationRuntime/Public/ConversationTaskNode.h
+++ b/Engine/Plugins/Experimental/CommonConversation/Source/CommonConversationRuntime/Public/ConversationTaskNode.h
@@ -82,6 +82,10 @@ public:
 
 	static void GenerateChoicesForDestinations(FConversationBranchPointBuilder& BranchBuilder, const FConversationContext& InContext, const TArray<FGuid>& CandidateDestinations);
 
+	/** Generates this task nodes' next possible conversation chioces. Default behavior handles finding static destination node choices and dynamic choices
+	  * generated by the node, but this can be overriden if more custom behavior is desired for a specific task node. */
+	virtual void GenerateNextChoices(FConversationBranchPointBuilder& BranchBuilder, const FConversationContext& Context) const;
+
 protected:
 	UFUNCTION(BlueprintNativeEvent, BlueprintAuthorityOnly)
 	EConversationRequirementResult IsRequirementSatisfied(const FConversationContext& Context) const;
