Common conversation - Update conditions under which we SendClientRefreshedTaskChoiceData, to test for more additional commonly changed dynamic data.
This ensures clients can be kept up to date, and do not become 'locked out' of advancing the conversation because their choice reference data sent with a request to advance, does not match the server's most up to date version.

diff --git a/Engine/Plugins/Experimental/CommonConversation/Source/CommonConversationRuntime/Private/ConversationParticipantComponent.cpp b/Engine/Plugins/Experimental/CommonConversation/Source/CommonConversationRuntime/Private/ConversationParticipantComponent.cpp
--- a/Engine/Plugins/Experimental/CommonConversation/Source/CommonConversationRuntime/Private/ConversationParticipantComponent.cpp
+++ b/Engine/Plugins/Experimental/CommonConversation/Source/CommonConversationRuntime/Private/ConversationParticipantComponent.cpp
@@ -210,11 +210,19 @@ void UConversationParticipantComponent::SendClientRefreshedTaskChoiceData(const
 		{
 			for (FClientConversationOptionEntry& LastOption : LastMessage.Options)
 			{
+				// UConversationInstance::FindBranchPointFromClientChoice will reject advancing a conversation from a client choice
+				// if InBranchPoint.ClientChoice == InChoice. This equality operation requires the client have the most up to date
+				// values for all fields. If *any* of these have changed, we should allow refresh. 
+				// If required we could eventually just send the entire FClientConversationOptionEntry here
+				// and process it similarly in ClientUpdateConversationTaskChoiceData_Implementation, but attempting
+				// to be somewhat selective in processing all 'known dynamic fields' while still possible for now to reduce overheads
 				if (LastOption.ChoiceReference.NodeReference == Handle && 
-					LastOption.ExtraData != CurrentOption.ExtraData)
+					(LastOption.ExtraData != CurrentOption.ExtraData || 
+					LastOption.ChoiceReference.NodeParameters != CurrentOption.ChoiceReference.NodeParameters))
 				{
 					LastOption.ExtraData = CurrentOption.ExtraData;
-					ClientUpdateConversationTaskChoiceData(Handle, LastOption);
+					LastOption.ChoiceReference.NodeParameters = CurrentOption.ChoiceReference.NodeParameters;
+					ClintUpdateConversationTaskChoiceData(Handle, LastOption);
 					break;
 				}
 			}
@@ -248,6 +256,7 @@ void UConversationParticipantComponent::ClientUpdateConversationTaskChoiceData_I
 		if (ExistingOption.ChoiceReference.NodeReference == Handle)
 		{
 			ExistingOption.ExtraData = OptionEntry.ExtraData;
+			ExistingOption.ChoiceReference.NodeParameters = OptionEntry.ChoiceReference.NodeParameters;
 			ConversationTaskChoiceDataUpdated.Broadcast(Handle, OptionEntry);
 			break;
 		}
