ConversationParticipantComponent.cpp - compilation fix for typo ofter merge conflict - "ClintUpdateConversationTaskChoiceData" -> "ClientUpdateConversationTaskChoiceData"

diff --git a/Engine/Plugins/Experimental/CommonConversation/Source/CommonConversationRuntime/Private/ConversationParticipantComponent.cpp b/Engine/Plugins/Experimental/CommonConversation/Source/CommonConversationRuntime/Private/ConversationParticipantComponent.cpp
--- a/Engine/Plugins/Experimental/CommonConversation/Source/CommonConversationRuntime/Private/ConversationParticipantComponent.cpp
+++ b/Engine/Plugins/Experimental/CommonConversation/Source/CommonConversationRuntime/Private/ConversationParticipantComponent.cpp
@@ -222,7 +222,7 @@ void UConversationParticipantComponent::SendClientRefreshedTaskChoiceData(const
 				{
 					LastOption.ExtraData = CurrentOption.ExtraData;
 					LastOption.ChoiceReference.NodeParameters = CurrentOption.ChoiceReference.NodeParameters;
-					ClintUpdateConversationTaskChoiceData(Handle, LastOption);
+					ClientUpdateConversationTaskChoiceData(Handle, LastOption);
 					break;
 				}
 			}
