Modified UConversationParticipantComponent::ServerNotifyConversationEnded so it looks for the conversation being passed in, instead of assuming it matches Auth_CurrentConversation.

There is a larger issue at play here regarding Auth_CurrentConversation. An NPC supports multiple conversations simultaneously, which seems very much at odds with the concept of Auth_CurrentConversation. I think there's more work to do here (ie removal of this variable), but it's going to have a large impact on methods dependent on Auth_CurrentConversation.

diff --git a/Engine/Plugins/Experimental/CommonConversation/Source/CommonConversationRuntime/Private/ConversationParticipantComponent.cpp b/Engine/Plugins/Experimental/CommonConversation/Source/CommonConversationRuntime/Private/ConversationParticipantComponent.cpp
--- a/Engine/Plugins/Experimental/CommonConversation/Source/CommonConversationRuntime/Private/ConversationParticipantComponent.cpp
+++ b/Engine/Plugins/Experimental/CommonConversation/Source/CommonConversationRuntime/Private/ConversationParticipantComponent.cpp
@@ -65,23 +65,33 @@ void UConversationParticipantComponent::ServerNotifyConversationEnded(UConversat
 	AActor* Owner = GetOwner();
 	if (Owner->GetLocalRole() == ROLE_Authority)
 	{
-		if (Auth_CurrentConversation == Conversation)
+		check(Conversation);
+
+		for (UConversationInstance* ConversationInstance : Auth_Conversations)
 		{
-			check(Conversation);
-			Auth_CurrentConversation = nullptr;
-			Auth_Conversations.Remove(Conversation);
+			if (Conversation == ConversationInstance)
+			{
+				if (Conversation == Auth_CurrentConversation)
+				{
+					Auth_CurrentConversation = nullptr;
+				}
+				
+				Auth_Conversations.Remove(Conversation);
 
-			const int32 OldConversationsActive = ConversationsActive;
-			ConversationsActive--;
-			MARK_PROPERTY_DIRTY_FROM_NAME(UConversationParticipantComponent, ConversationsActive, this);
+				const int32 OldConversationsActive = ConversationsActive;
+				ConversationsActive--;
+				MARK_PROPERTY_DIRTY_FROM_NAME(UConversationParticipantComponent, ConversationsActive, this);
 
-			OnServerConversationEnded(Conversation);
+				OnServerConversationEnded(Conversation);
 
-			ClientUpdateConversations(ConversationsActive);
+				ClientUpdateConversations(ConversationsActive);
 
-			if (ConversationsActive == 0)
-			{
-				OnRep_ConversationsActive(OldConversationsActive);
+				if (ConversationsActive == 0)
+				{
+					OnRep_ConversationsActive(OldConversationsActive);
+				}
+
+				break;
 			}
 		}
 	}
