QS2 - Conversation Choice Made Verb - Send a "Conversation Choice Made" Verb when a player notifies the server that a particular choice has been made, if that choice node contains any choice tags which could be relevant to external users.

diff --git a/Engine/Plugins/Experimental/CommonConversation/Source/CommonConversationRuntime/Private/ConversationInstance.cpp b/Engine/Plugins/Experimental/CommonConversation/Source/CommonConversationRuntime/Private/ConversationInstance.cpp
--- a/Engine/Plugins/Experimental/CommonConversation/Source/CommonConversationRuntime/Private/ConversationInstance.cpp
+++ b/Engine/Plugins/Experimental/CommonConversation/Source/CommonConversationRuntime/Private/ConversationInstance.cpp
@@ -193,6 +193,9 @@ void UConversationInstance::ServerAdvanceConversation(const FAdvanceConversation
 
 		TArray<FConversationBranchPoint> CandidateDestinations;
 
+		const FConversationContext ServerContext = FConversationContext::CreateServerContext(this, nullptr);
+		const UConversationChoiceNode* ChoiceNodePicked = nullptr;
+
 		if (InChoicePicked.Choice != FConversationChoiceReference::Empty)
 		{
 			if (const FConversationBranchPoint* BranchPoint = FindBranchPointFromClientChoice(InChoicePicked.Choice))
@@ -200,14 +203,14 @@ void UConversationInstance::ServerAdvanceConversation(const FAdvanceConversation
 				UE_LOG(LogCommonConversationRuntime, Verbose, TEXT("User picked option %s, going to try that"), *BranchPoint->ClientChoice.ChoiceReference.ToString());
 				CandidateDestinations = { *BranchPoint };
 
-				FConversationContext Context = FConversationContext::CreateServerContext(this, nullptr);
-				if (const UConversationTaskNode* TaskNode = BranchPoint->ClientChoice.TryToResolveChoiceNode<UConversationTaskNode>(Context))
+				if (const UConversationTaskNode* TaskNode = BranchPoint->ClientChoice.TryToResolveChoiceNode<UConversationTaskNode>(ServerContext))
 				{
 					for (const UConversationNode* SubNode : TaskNode->SubNodes)
 					{
 						if (const UConversationChoiceNode* ChoiceNode = Cast<UConversationChoiceNode>(SubNode))
 						{
-							ChoiceNode->NotifyChoicePickedByUser(Context, BranchPoint->ClientChoice);
+							ChoiceNodePicked = ChoiceNode;
+							ChoiceNode->NotifyChoicePickedByUser(ServerContext, BranchPoint->ClientChoice);
 							break;
 						}
 					}
@@ -257,6 +260,12 @@ void UConversationInstance::ServerAdvanceConversation(const FAdvanceConversation
 			}
 		}
 
+		// Allow derived conversation instances a chance to respond to a choice being picked
+		if (ChoiceNodePicked)
+		{ 
+			OnChoiceNodePickedByUser(ServerContext, ChoiceNodePicked, ValidDestinations);
+		}
+
 		if (ValidDestinations.Num() == 0)
 		{
 			UE_LOG(LogCommonConversationRuntime, Verbose, TEXT("No available destinations from %s, ending the conversation"), *GetCurrentChoiceReference().ToString());
diff --git a/Engine/Plugins/Experimental/CommonConversation/Source/CommonConversationRuntime/Public/ConversationInstance.h b/Engine/Plugins/Experimental/CommonConversation/Source/CommonConversationRuntime/Public/ConversationInstance.h
--- a/Engine/Plugins/Experimental/CommonConversation/Source/CommonConversationRuntime/Public/ConversationInstance.h
+++ b/Engine/Plugins/Experimental/CommonConversation/Source/CommonConversationRuntime/Public/ConversationInstance.h
@@ -9,6 +9,7 @@
 #include "Math/RandomStream.h"
 #include "ConversationInstance.generated.h"
 
+class UConversationChoiceNode;
 class UConversationInstance;
 class UConversationNodeWithLinks;
 class UConversationParticipantComponent;
@@ -98,6 +99,7 @@ protected:
 	void ReturnToCurrentClientChoice(const FConversationContext& Context);
 	void ReturnToStart(const FConversationContext& Context);
 	virtual void PauseConversationAndSendClientChoices(const FConversationContext& Context, const FClientConversationMessage& ClientMessage);
+	virtual void OnChoiceNodePickedByUser(const FConversationContext& Context, const UConversationChoiceNode* ChoiceNode, const TArray<FConversationBranchPoint>& ValidDestinations) {};
 #endif
 
 private:
