Engine:
- Modified the ClientStartConversation rpc to contain information about conversation participants
- Added a ClientExitConversation rpc with information about the conversation participants

diff --git a/Engine/Plugins/Experimental/CommonConversation/Source/CommonConversationRuntime/Private/ConversationInstance.cpp b/Engine/Plugins/Experimental/CommonConversation/Source/CommonConversationRuntime/Private/ConversationInstance.cpp
--- a/Engine/Plugins/Experimental/CommonConversation/Source/CommonConversationRuntime/Private/ConversationInstance.cpp
+++ b/Engine/Plugins/Experimental/CommonConversation/Source/CommonConversationRuntime/Private/ConversationInstance.cpp
@@ -26,7 +26,7 @@ UWorld* UConversationInstance::GetWorld() const
 
 #if WITH_SERVER_CODE
 
-void UConversationInstance::ServerRemoveParticipant(FGameplayTag ParticipantID)
+void UConversationInstance::ServerRemoveParticipant(FGameplayTag ParticipantID, const FConversationParticipants& PreservedParticipants)
 {
 	for (auto It = Participants.List.CreateIterator(); It; ++It)
 	{
@@ -36,7 +36,7 @@ void UConversationInstance::ServerRemoveParticipant(FGameplayTag ParticipantID)
 			{
 				if (bConversationStarted)
 				{
-					OldParticipant->ServerNotifyConversationEnded(this);
+					OldParticipant->ServerNotifyConversationEnded(this, PreservedParticipants);
 				}
 			}
 			It.RemoveCurrent();
@@ -54,7 +54,7 @@ void UConversationInstance::ServerAssignParticipant(FGameplayTag ParticipantID,
 		return;
 	}
 
-	ServerRemoveParticipant(ParticipantID);
+	ServerRemoveParticipant(ParticipantID, Participants);
 
 	FConversationParticipantEntry NewEntry;
 	NewEntry.ParticipantID = ParticipantID;
@@ -304,10 +304,10 @@ void UConversationInstance::ServerAbortConversation()
 
 		OnEnded();
 
-		TArray<FConversationParticipantEntry> ListCopy = Participants.List;
-		for (const FConversationParticipantEntry& ParticipantEntry : ListCopy)
+		const FConversationParticipants ParticipantsCopy = Participants;
+		for (const FConversationParticipantEntry& ParticipantEntry : ParticipantsCopy.List)
 		{
-			ServerRemoveParticipant(ParticipantEntry.ParticipantID);
+			ServerRemoveParticipant(ParticipantEntry.ParticipantID, ParticipantsCopy);
 		}
 		check(Participants.List.Num() == 0);
 	}
diff --git a/Engine/Plugins/Experimental/CommonConversation/Source/CommonConversationRuntime/Private/ConversationParticipantComponent.cpp b/Engine/Plugins/Experimental/CommonConversation/Source/CommonConversationRuntime/Private/ConversationParticipantComponent.cpp
--- a/Engine/Plugins/Experimental/CommonConversation/Source/CommonConversationRuntime/Private/ConversationParticipantComponent.cpp
+++ b/Engine/Plugins/Experimental/CommonConversation/Source/CommonConversationRuntime/Private/ConversationParticipantComponent.cpp
@@ -54,13 +54,13 @@ void UConversationParticipantComponent::ServerNotifyConversationStarted(UConvers
 		MARK_PROPERTY_DIRTY_FROM_NAME(UConversationParticipantComponent, ConversationsActive, this);
 
 		OnServerConversationStarted(Conversation, AsParticipant);
-		ClientStartConversation(AsParticipant);
+		ClientStartConversation(Conversation->GetParticipantsCopy());
 
 		ClientUpdateConversations(ConversationsActive);
 	}
 }
 
-void UConversationParticipantComponent::ServerNotifyConversationEnded(UConversationInstance* Conversation)
+void UConversationParticipantComponent::ServerNotifyConversationEnded(UConversationInstance* Conversation, const FConversationParticipants& PreservedParticipants)
 {
 	AActor* Owner = GetOwner();
 	if (Owner->GetLocalRole() == ROLE_Authority)
@@ -83,6 +83,7 @@ void UConversationParticipantComponent::ServerNotifyConversationEnded(UConversat
 				MARK_PROPERTY_DIRTY_FROM_NAME(UConversationParticipantComponent, ConversationsActive, this);
 
 				OnServerConversationEnded(Conversation);
+				ClientExitConversation(PreservedParticipants);
 
 				ClientUpdateConversations(ConversationsActive);
 
@@ -273,10 +274,17 @@ void UConversationParticipantComponent::ClientUpdateConversationTaskChoiceData_I
 	}
 }
 
-void UConversationParticipantComponent::ClientStartConversation_Implementation(const FGameplayTag AsParticipant)
+void UConversationParticipantComponent::ClientStartConversation_Implementation(const FConversationParticipants& InParticipants)
 {
 	bIsFirstConversationUpdateBroadcasted = false;
 	ConversationStarted.Broadcast();
+
+	OnClientStartConversation(InParticipants);
+}
+
+void UConversationParticipantComponent::ClientExitConversation_Implementation(const FConversationParticipants& InParticipants)
+{
+	OnClientExitConversation(InParticipants);
 }
 
 bool UConversationParticipantComponent::IsInActiveConversation() const
@@ -368,6 +376,16 @@ void UConversationParticipantComponent::OnConversationUpdated(const FClientConve
 
 }
 
+void UConversationParticipantComponent::OnClientStartConversation(const FConversationParticipants& InParticipants)
+{
+
+}
+
+void UConversationParticipantComponent::OnClientExitConversation(const FConversationParticipants& InParticipants)
+{
+
+}
+
 #if WITH_SERVER_CODE
 
 void UConversationParticipantComponent::ServerAbortAllConversations()
diff --git a/Engine/Plugins/Experimental/CommonConversation/Source/CommonConversationRuntime/Public/ConversationInstance.h b/Engine/Plugins/Experimental/CommonConversation/Source/CommonConversationRuntime/Public/ConversationInstance.h
--- a/Engine/Plugins/Experimental/CommonConversation/Source/CommonConversationRuntime/Public/ConversationInstance.h
+++ b/Engine/Plugins/Experimental/CommonConversation/Source/CommonConversationRuntime/Public/ConversationInstance.h
@@ -38,7 +38,9 @@ public:
 	virtual UWorld* GetWorld() const override;
 	
 #if WITH_SERVER_CODE
-	void ServerRemoveParticipant(FGameplayTag ParticipantID);
+	/** Should be called with a copy of the conversation participants before any removals happen, that way clients can properly respond to the end of their respective conversations
+	  * with an accurate account of who was in that conversation. */
+	void ServerRemoveParticipant(FGameplayTag ParticipantID, const FConversationParticipants& PreservedParticipants);
 
 	void ServerAssignParticipant(FGameplayTag ParticipantID, AActor* ParticipantActor);
 
diff --git a/Engine/Plugins/Experimental/CommonConversation/Source/CommonConversationRuntime/Public/ConversationParticipantComponent.h b/Engine/Plugins/Experimental/CommonConversation/Source/CommonConversationRuntime/Public/ConversationParticipantComponent.h
--- a/Engine/Plugins/Experimental/CommonConversation/Source/CommonConversationRuntime/Public/ConversationParticipantComponent.h
+++ b/Engine/Plugins/Experimental/CommonConversation/Source/CommonConversationRuntime/Public/ConversationParticipantComponent.h
@@ -47,7 +47,7 @@ public:
 
 #if WITH_SERVER_CODE
 	void ServerNotifyConversationStarted(UConversationInstance* Conversation, FGameplayTag AsParticipant);
-	void ServerNotifyConversationEnded(UConversationInstance* Conversation);
+	void ServerNotifyConversationEnded(UConversationInstance* Conversation, const FConversationParticipants& PreservedParticipants);
 	void ServerNotifyExecuteTaskAndSideEffects(const FConversationNodeHandle& Handle);
 	void ServerForAllConversationsRefreshChoices(UConversationInstance* IgnoreConversation = nullptr);
 	void ServerForAllConversationsRefreshTaskChoiceData(const FConversationNodeHandle& Handle, UConversationInstance* IgnoreConversation /*= nullptr*/);
@@ -103,7 +103,10 @@ protected:
 	void ClientUpdateConversations(int32 InConversationsActive);
 
 	UFUNCTION(Client, Reliable)
-	void ClientStartConversation(const FGameplayTag AsParticipant);
+	void ClientStartConversation(const FConversationParticipants& InParticipants);
+
+	UFUNCTION(Client, Reliable)
+	void ClientExitConversation(const FConversationParticipants& InParticipants);
 
 protected:
 	UFUNCTION()
@@ -113,6 +116,9 @@ protected:
 	virtual void OnLeaveConversationState();
 	virtual void OnConversationUpdated(const FClientConversationMessagePayload& Message);
 
+	virtual void OnClientStartConversation(const FConversationParticipants& InParticipants);
+	virtual void OnClientExitConversation(const FConversationParticipants& InParticipants);
+
 #if WITH_SERVER_CODE
 	virtual void OnServerConversationStarted(UConversationInstance* Conversation, FGameplayTag AsParticipant);
 	virtual void OnServerConversationEnded(UConversationInstance* Conversation);
