Lock support for conversation services

diff --git a/Engine/Plugins/Experimental/CommonConversation/Source/CommonConversationGraph/Private/ConversationCompiler.cpp b/Engine/Plugins/Experimental/CommonConversation/Source/CommonConversationGraph/Private/ConversationCompiler.cpp
--- a/Engine/Plugins/Experimental/CommonConversation/Source/CommonConversationGraph/Private/ConversationCompiler.cpp
+++ b/Engine/Plugins/Experimental/CommonConversation/Source/CommonConversationGraph/Private/ConversationCompiler.cpp
@@ -220,8 +220,16 @@ void FConversationCompiler::RebuildBank(UConversationDatabase* ConversationAsset
 			UEdGraphPin* OutputPin = EdNode->GetOutputPin();
 			check(OutputPin);
 
-			ForeachConnectedOutgoingConversationNode(OutputPin, [RuntimeNodeWithLinks](UConversationGraphNode* RemoteNode) {
+			ForeachConnectedOutgoingConversationNode(OutputPin, [RuntimeNodeWithLinks](UConversationGraphNode* RemoteNode)
+			{
 				RuntimeNodeWithLinks->OutputConnections.Add(RemoteNode->NodeGuid);
+				if (UConversationNodeWithLinks* ChildNodeWithLinks = Cast<UConversationNodeWithLinks>(RemoteNode->NodeInstance))
+				{
+					// Tell child about its parent so calls to 'GetParentNode' work as expected
+					// Could alternatively pass as FGuid and store similarly to 'output connections' (from RuntimeNodeWithLinks->GetNodeGuid()), 
+					// Users would need to call UConversationRegistry::GetRuntimeNodeFromGUID when they want the parent node in that case
+					ChildNodeWithLinks->InitializeNode(RuntimeNodeWithLinks);
+				}
 			});
 		}
 	}
