Common Conversation Runtime: Adding logs to UConversationPartipantComponent::EnterConversationState and UConversationPartipantComponent::ExitConversationState to assist with debugging conversation flow issues

diff --git a/Engine/Plugins/Experimental/CommonConversation/Source/CommonConversationRuntime/Private/ConversationParticipantComponent.cpp b/Engine/Plugins/Experimental/CommonConversation/Source/CommonConversationRuntime/Private/ConversationParticipantComponent.cpp
--- a/Engine/Plugins/Experimental/CommonConversation/Source/CommonConversationRuntime/Private/ConversationParticipantComponent.cpp
+++ b/Engine/Plugins/Experimental/CommonConversation/Source/CommonConversationRuntime/Private/ConversationParticipantComponent.cpp
@@ -349,7 +349,7 @@ void UConversationParticipantComponent::OnRep_ConversationsActive(int32 OldConve
 
 void UConversationParticipantComponent::OnEnterConversationState()
 {
-
+	UE_LOG(LogCommonConversationRuntime, Verbose, TEXT("[%s]: %s has entered the conversation state."), ANSI_TO_TCHAR(__FUNCTION__), *GetNameSafe(GetOwner()));
 }
 
 #if WITH_SERVER_CODE
@@ -368,7 +368,7 @@ void UConversationParticipantComponent::OnServerConversationEnded(UConversationI
 
 void UConversationParticipantComponent::OnLeaveConversationState()
 {
-
+	UE_LOG(LogCommonConversationRuntime, Verbose, TEXT("[%s]: %s has exited the conversation state."), ANSI_TO_TCHAR(__FUNCTION__), *GetNameSafe(GetOwner()));
 }
 
 void UConversationParticipantComponent::OnConversationUpdated(const FClientConversationMessagePayload& Message)
