Common Conversation Plugin:
- Changing ExitTag generation when compiling conversation graph assets to use `AddTag` instead of `AddLeafTag`. Using the leaf version was causing issues for conversation graphs that used both a a child and a parent tag for two different exit points. The result was the parent tag exit point getting ignored and thus that exit point not counting as a dependency for the given conversation asset, which could butterfly effect into sync loads because the given graph is not pre-loaded properly when readying a conversation component.
- Affected Conversation Graphs will need to be re-saved for the fix to take effect.

#tests PIE 28.10

diff --git a/Engine/Plugins/Experimental/CommonConversation/Source/CommonConversationGraph/Private/ConversationCompiler.cpp b/Engine/Plugins/Experimental/CommonConversation/Source/CommonConversationGraph/Private/ConversationCompiler.cpp
--- a/Engine/Plugins/Experimental/CommonConversation/Source/CommonConversationGraph/Private/ConversationCompiler.cpp
+++ b/Engine/Plugins/Experimental/CommonConversation/Source/CommonConversationGraph/Private/ConversationCompiler.cpp
@@ -307,7 +307,7 @@ void FConversationCompiler::RebuildBank(UConversationDatabase* ConversationAsset
 
 		if (UConversationLinkNode* ExitTagNode = Cast<UConversationLinkNode>(NodeInstance))
 		{
-			ConversationAsset->ExitTags.AddLeafTag(ExitTagNode->GetRemoteEntryTag());
+			ConversationAsset->ExitTags.AddTag(ExitTagNode->GetRemoteEntryTag());
 		}
 	}
 
