Conversations - exposing query to acess exit link functionality and better understand relationships between dependent conversation graphs

diff --git a/Engine/Plugins/Experimental/CommonConversation/Source/CommonConversationRuntime/Private/ConversationRegistry.cpp b/Engine/Plugins/Experimental/CommonConversation/Source/CommonConversationRuntime/Private/ConversationRegistry.cpp
--- a/Engine/Plugins/Experimental/CommonConversation/Source/CommonConversationRuntime/Private/ConversationRegistry.cpp
+++ b/Engine/Plugins/Experimental/CommonConversation/Source/CommonConversationRuntime/Private/ConversationRegistry.cpp
@@ -358,6 +358,15 @@ TArray<FPrimaryAssetId> UConversationRegistry::GetPrimaryAssetIdsForEntryPoint(F
 	return AssetsWithTheEntryPoint;
 }
 
+TArray<FGameplayTag> UConversationRegistry::GetLinkedExitConversationEntryTags(const UConversationDatabase* ConversationDatabase) const
+{
+	if (ConversationDatabase)
+	{
+		return ConversationDatabase->ExitTags.GetGameplayTagArray();
+	}
+	return {};
+}
+
 void UConversationRegistry::BuildDependenciesGraph()
 {
 	if (bDependenciesBuilt)
diff --git a/Engine/Plugins/Experimental/CommonConversation/Source/CommonConversationRuntime/Public/ConversationRegistry.h b/Engine/Plugins/Experimental/CommonConversation/Source/CommonConversationRuntime/Public/ConversationRegistry.h
--- a/Engine/Plugins/Experimental/CommonConversation/Source/CommonConversationRuntime/Public/ConversationRegistry.h
+++ b/Engine/Plugins/Experimental/CommonConversation/Source/CommonConversationRuntime/Public/ConversationRegistry.h
@@ -93,7 +93,10 @@ public:
 	TSharedPtr<FConversationsHandle> LoadConversationsFor(const TArray<FGameplayTag>& ConversationEntryTags) const;
 
 	TArray<FPrimaryAssetId> GetPrimaryAssetIdsForEntryPoint(FGameplayTag EntryPoint) const;
-	
+
+	// If a conversation database links to other conversaton assets, the tags of those conversations can be obtained here
+	TArray<FGameplayTag> GetLinkedExitConversationEntryTags(const UConversationDatabase* ConversationDatabase) const;
+
 	UPROPERTY(Transient)
 	FNetSerializeScriptStructCache_ConvVersion ConversationChoiceDataStructCache;
 
