Updates to CommonConversation plugin to make the system more customizable by projects.

   * Move the conversation asset type to the public folder.
   * AddConversationNodeOptions in the schema is changed to protected.
   * Added a function to start a conversation from a specific graph.

diff --git a/Engine/Plugins/Experimental/CommonConversation/Source/CommonConversationEditor/Private/AssetTypeActions_ConversationDatabase.h b/Engine/Plugins/Experimental/CommonConversation/Source/CommonConversationEditor/Private/AssetTypeActions_ConversationDatabase.h
deleted file mode 100644
diff --git a/Engine/Plugins/Experimental/CommonConversation/Source/CommonConversationEditor/Public/AssetTypeActions_ConversationDatabase.h b/Engine/Plugins/Experimental/CommonConversation/Source/CommonConversationEditor/Public/AssetTypeActions_ConversationDatabase.h
new file mode 100644
--- /dev/null
+++ b/Engine/Plugins/Experimental/CommonConversation/Source/CommonConversationEditor/Public/AssetTypeActions_ConversationDatabase.h
@@ -0,0 +1,24 @@
+// Copyright Epic Games, Inc. All Rights Reserved.
+
+#pragma once
+
+#include "AssetTypeActions_Base.h"
+
+class UConversationDatabase;
+
+class COMMONCONVERSATIONEDITOR_API FAssetTypeActions_ConversationDatabase : public FAssetTypeActions_Base
+{
+public:
+	// IAssetTypeActions Implementation
+	virtual FText GetName() const override { return NSLOCTEXT("AssetTypeActions", "AssetTypeActions_ConversationDatabase", "Conversation Bank"); }
+	virtual FColor GetTypeColor() const override { return FColor(149,70,255); }
+	virtual UClass* GetSupportedClass() const override;
+	virtual void OpenAssetEditor( const TArray<UObject*>& InObjects, TSharedPtr<class IToolkitHost> EditWithinLevelEditor = TSharedPtr<IToolkitHost>() ) override;
+	virtual uint32 GetCategories() override;
+	virtual void PerformAssetDiff(UObject* Asset1, UObject* Asset2, const struct FRevisionInfo& OldRevision, const struct FRevisionInfo& NewRevision) const override;
+
+private:
+
+	/* Called to open the Behavior Tree defaults view, this opens whatever text diff tool the user has */
+	void OpenInDefaults(UConversationDatabase* OldBank, UConversationDatabase* NewBank) const;
+};
diff --git a/Engine/Plugins/Experimental/CommonConversation/Source/CommonConversationGraph/Public/ConversationGraphSchema.h b/Engine/Plugins/Experimental/CommonConversation/Source/CommonConversationGraph/Public/ConversationGraphSchema.h
--- a/Engine/Plugins/Experimental/CommonConversation/Source/CommonConversationGraph/Public/ConversationGraphSchema.h
+++ b/Engine/Plugins/Experimental/CommonConversation/Source/CommonConversationGraph/Public/ConversationGraphSchema.h
@@ -62,7 +62,7 @@ public:
 	virtual void GetGraphNodeContextActions(FGraphContextMenuBuilder& ContextMenuBuilder, int32 SubNodeFlags) const override;
 	virtual void GetSubNodeClasses(int32 SubNodeFlags, TArray<FGraphNodeClassData>& ClassData, UClass*& GraphNodeClass) const override;
 
-private:
+protected:
 	void AddConversationNodeOptions(const FString& CategoryName, FGraphContextMenuBuilder& ContextMenuBuilder, TSubclassOf<UConversationNode> RuntimeNodeType, TSubclassOf<UConversationGraphNode> EditorNodeType) const;
 
 private:
diff --git a/Engine/Plugins/Experimental/CommonConversation/Source/CommonConversationRuntime/Private/ConversationInstance.cpp b/Engine/Plugins/Experimental/CommonConversation/Source/CommonConversationRuntime/Private/ConversationInstance.cpp
--- a/Engine/Plugins/Experimental/CommonConversation/Source/CommonConversationRuntime/Private/ConversationInstance.cpp
+++ b/Engine/Plugins/Experimental/CommonConversation/Source/CommonConversationRuntime/Private/ConversationInstance.cpp
@@ -26,7 +26,7 @@ UWorld* UConversationInstance::GetWorld() const
 
 #if WITH_SERVER_CODE
 
-void UConversationInstance::ServerRemoveParticipant(FGameplayTag ParticipantID, const FConversationParticipants& PreservedParticipants)
+void UConversationInstance::ServerRemoveParticipant(const FGameplayTag& ParticipantID, const FConversationParticipants& PreservedParticipants)
 {
 	for (auto It = Participants.List.CreateIterator(); It; ++It)
 	{
@@ -45,7 +45,7 @@ void UConversationInstance::ServerRemoveParticipant(FGameplayTag ParticipantID,
 	}
 }
 
-void UConversationInstance::ServerAssignParticipant(FGameplayTag ParticipantID, AActor* ParticipantActor)
+void UConversationInstance::ServerAssignParticipant(const FGameplayTag& ParticipantID, AActor* ParticipantActor)
 {
 	if (!ParticipantID.IsValid() || (ParticipantActor == nullptr))
 	{
@@ -73,7 +73,7 @@ void UConversationInstance::ServerAssignParticipant(FGameplayTag ParticipantID,
 		*GetName(), *ParticipantID.ToString(), *GetPathNameSafe(ParticipantActor));
 }
 
-void UConversationInstance::ServerStartConversation(FGameplayTag EntryPoint)
+void UConversationInstance::ServerStartConversation(const FGameplayTag& EntryPoint, const UConversationDatabase* Graph)
 {
 	UE_LOG(LogCommonConversationRuntime, Verbose, TEXT("Conversation %s starting at %s with %d participants"),
 		*GetName(), *EntryPoint.ToString(), Participants.List.Num());
@@ -83,7 +83,7 @@ void UConversationInstance::ServerStartConversation(FGameplayTag EntryPoint)
 
 	UConversationRegistry* ConversationRegistry = UConversationRegistry::GetFromWorld(GetWorld());
 	
-	TArray<FGuid> PotentialStartingPoints = ConversationRegistry->GetOutputLinkGUIDs(EntryPoint);
+	TArray<FGuid> PotentialStartingPoints = ConversationRegistry->GetOutputLinkGUIDs(Graph, EntryPoint);
 	if (PotentialStartingPoints.Num() == 0)
 	{
 		UE_LOG(LogCommonConversationRuntime, Warning, TEXT("Entry point %s did not exist or had no destination entries; conversation aborted"), *EntryPoint.ToString());
diff --git a/Engine/Plugins/Experimental/CommonConversation/Source/CommonConversationRuntime/Private/ConversationLibrary.cpp b/Engine/Plugins/Experimental/CommonConversation/Source/CommonConversationRuntime/Private/ConversationLibrary.cpp
--- a/Engine/Plugins/Experimental/CommonConversation/Source/CommonConversationRuntime/Private/ConversationLibrary.cpp
+++ b/Engine/Plugins/Experimental/CommonConversation/Source/CommonConversationRuntime/Private/ConversationLibrary.cpp
@@ -9,18 +9,11 @@
 
 #define LOCTEXT_NAMESPACE "ConversationLibrary"
 
-//////////////////////////////////////////////////////////////////////////
-// UConversationLibrary
-
-UConversationLibrary::UConversationLibrary()
-{
-}
-
-UConversationInstance* UConversationLibrary::StartConversation(FGameplayTag ConversationEntryTag, AActor* Instigator,
-	FGameplayTag InstigatorTag, AActor* Target, FGameplayTag TargetTag, const TSubclassOf<UConversationInstance> ConversationInstanceClass)
+static UConversationInstance* StartConversationShared(const FGameplayTag& ConversationEntryTag, AActor* Instigator, const FGameplayTag& InstigatorTag,
+	AActor* Target, const FGameplayTag& TargetTag, const TSubclassOf<UConversationInstance> ConversationInstanceClass, const UConversationDatabase* Graph)
 {
 #if WITH_SERVER_CODE
-	if (Instigator == nullptr || Target == nullptr)
+	if (Instigator == nullptr || Target == nullptr || GEngine == nullptr)
 	{
 		return nullptr;
 	}
@@ -44,7 +37,7 @@ UConversationInstance* UConversationLibrary::StartConversation(FGameplayTag Conv
 			UConversationContextHelpers::MakeConversationParticipant(Context, Target, TargetTag);
 			UConversationContextHelpers::MakeConversationParticipant(Context, Instigator, InstigatorTag);
 
-			ConversationInstance->ServerStartConversation(ConversationEntryTag);
+			ConversationInstance->ServerStartConversation(ConversationEntryTag, Graph);
 		}
 
 		return ConversationInstance;
@@ -54,5 +47,24 @@ UConversationInstance* UConversationLibrary::StartConversation(FGameplayTag Conv
 	return nullptr;
 }
 
+//////////////////////////////////////////////////////////////////////////
+// UConversationLibrary
+
+UConversationLibrary::UConversationLibrary()
+{
+}
+
+UConversationInstance* UConversationLibrary::StartConversation(const FGameplayTag& ConversationEntryTag, AActor* Instigator,
+	const FGameplayTag& InstigatorTag, AActor* Target, const FGameplayTag& TargetTag, const TSubclassOf<UConversationInstance> ConversationInstanceClass)
+{
+	return StartConversationShared(ConversationEntryTag, Instigator, InstigatorTag, Target, TargetTag, ConversationInstanceClass, nullptr);
+}
+
+UConversationInstance* UConversationLibrary::StartConversationFromGraph(const FGameplayTag& ConversationEntryTag, AActor* Instigator, const FGameplayTag& InstigatorTag,
+	AActor* Target, const FGameplayTag& TargetTag, const UConversationDatabase* Graph)
+{
+	return (Graph) ? StartConversationShared(ConversationEntryTag, Instigator, InstigatorTag, Target, TargetTag, nullptr, Graph) : nullptr;
+}
+
 #undef LOCTEXT_NAMESPACE
 
diff --git a/Engine/Plugins/Experimental/CommonConversation/Source/CommonConversationRuntime/Private/ConversationRegistry.cpp b/Engine/Plugins/Experimental/CommonConversation/Source/CommonConversationRuntime/Private/ConversationRegistry.cpp
--- a/Engine/Plugins/Experimental/CommonConversation/Source/CommonConversationRuntime/Private/ConversationRegistry.cpp
+++ b/Engine/Plugins/Experimental/CommonConversation/Source/CommonConversationRuntime/Private/ConversationRegistry.cpp
@@ -210,14 +210,14 @@ UConversationNode* UConversationRegistry::GetRuntimeNodeFromGUID(const FGuid& No
 	return nullptr;
 }
 
-TArray<FGuid> UConversationRegistry::GetEntryPointGUIDs(FGameplayTag EntryPoint) const
+TArray<FGuid> UConversationRegistry::GetEntryPointGUIDs(const FGameplayTag& EntryPoint) const
 {
 	const_cast<UConversationRegistry*>(this)->BuildDependenciesGraph();
 
 	return EntryTagToEntryList.FindRef(EntryPoint);
 }
 
-TArray<FGuid> UConversationRegistry::GetOutputLinkGUIDs(FGameplayTag EntryPoint) const
+TArray<FGuid> UConversationRegistry::GetOutputLinkGUIDs(const FGameplayTag& EntryPoint) const
 {
 	TArray<FGuid> SourceGUIDs = GetEntryPointGUIDs(EntryPoint);
 	return GetOutputLinkGUIDs(SourceGUIDs);
@@ -239,7 +239,7 @@ TArray<FGuid> UConversationRegistry::GetOutputLinkGUIDs(const TArray<FGuid>& Sou
 		if (const UConversationDatabase* SourceConversation = GetConversationFromNodeGUID(SourceGUID))
 		{
 			UConversationNode* SourceNode = SourceConversation->ReachableNodeMap.FindRef(SourceGUID);
-			if (ensure(SourceNode))
+			if (SourceNode)
 			{
 				if (UConversationNodeWithLinks* SourceNodeWithLinks = CastChecked<UConversationNodeWithLinks>(SourceNode))
 				{
@@ -252,6 +252,37 @@ TArray<FGuid> UConversationRegistry::GetOutputLinkGUIDs(const TArray<FGuid>& Sou
 	return Result;
 }
 
+TArray<FGuid> UConversationRegistry::GetOutputLinkGUIDs(const UConversationDatabase* Graph, const FGameplayTag& EntryPoint) const
+{
+	if (Graph == nullptr)
+	{
+		return GetOutputLinkGUIDs(EntryPoint);
+	}
+
+	TArray<FGuid> Result;
+
+	for (const FConversationEntryList& Tag : Graph->EntryTags)
+	{
+		if (Tag.EntryTag == EntryPoint)
+		{
+			for (const FGuid& Destination : Tag.DestinationList)
+			{
+				UConversationNode* SourceNode = Graph->ReachableNodeMap.FindRef(Destination);
+				if (SourceNode)
+				{
+					if (UConversationNodeWithLinks* SourceNodeWithLinks = CastChecked<UConversationNodeWithLinks>(SourceNode))
+					{
+						Result.Append(SourceNodeWithLinks->OutputConnections);
+					}
+				}
+			}
+			break;
+		}
+	}
+
+	return Result;
+}
+
 UConversationDatabase* UConversationRegistry::GetConversationFromNodeGUID(const FGuid& NodeGUID) const
 {
 	if (const FSoftObjectPath* ConversationPathPtr = NodeGuidToConversation.Find(NodeGUID))
diff --git a/Engine/Plugins/Experimental/CommonConversation/Source/CommonConversationRuntime/Public/ConversationInstance.h b/Engine/Plugins/Experimental/CommonConversation/Source/CommonConversationRuntime/Public/ConversationInstance.h
--- a/Engine/Plugins/Experimental/CommonConversation/Source/CommonConversationRuntime/Public/ConversationInstance.h
+++ b/Engine/Plugins/Experimental/CommonConversation/Source/CommonConversationRuntime/Public/ConversationInstance.h
@@ -10,6 +10,7 @@
 #include "ConversationInstance.generated.h"
 
 class UConversationChoiceNode;
+class UConversationDatabase;
 class UConversationInstance;
 class UConversationNodeWithLinks;
 class UConversationParticipantComponent;
@@ -40,11 +41,11 @@ public:
 #if WITH_SERVER_CODE
 	/** Should be called with a copy of the conversation participants before any removals happen, that way clients can properly respond to the end of their respective conversations
 	  * with an accurate account of who was in that conversation. */
-	void ServerRemoveParticipant(FGameplayTag ParticipantID, const FConversationParticipants& PreservedParticipants);
+	void ServerRemoveParticipant(const FGameplayTag& ParticipantID, const FConversationParticipants& PreservedParticipants);
 
-	void ServerAssignParticipant(FGameplayTag ParticipantID, AActor* ParticipantActor);
+	void ServerAssignParticipant(const FGameplayTag& ParticipantID, AActor* ParticipantActor);
 
-	void ServerStartConversation(FGameplayTag EntryPoint);
+	void ServerStartConversation(const FGameplayTag& EntryPoint, const UConversationDatabase* Graph = nullptr);
 
 	void ServerAdvanceConversation(const FAdvanceConversationRequest& InChoicePicked);
 
diff --git a/Engine/Plugins/Experimental/CommonConversation/Source/CommonConversationRuntime/Public/ConversationLibrary.h b/Engine/Plugins/Experimental/CommonConversation/Source/CommonConversationRuntime/Public/ConversationLibrary.h
--- a/Engine/Plugins/Experimental/CommonConversation/Source/CommonConversationRuntime/Public/ConversationLibrary.h
+++ b/Engine/Plugins/Experimental/CommonConversation/Source/CommonConversationRuntime/Public/ConversationLibrary.h
@@ -11,6 +11,7 @@
 struct FGameplayTag;
 
 class AActor;
+class UConversationDatabase;
 class UConversationInstance;
 
 UCLASS()
@@ -22,8 +23,12 @@ public:
 	UConversationLibrary();
 
 	UFUNCTION(BlueprintCallable, BlueprintAuthorityOnly, Category="Conversation", meta = (DeterminesOutputType = "ConversationType"))
-	static UConversationInstance* StartConversation(FGameplayTag ConversationEntryTag, AActor* Instigator, FGameplayTag InstigatorTag,
-		AActor* Target, FGameplayTag TargetTag, const TSubclassOf<UConversationInstance> ConversationInstanceClass = nullptr);
+	static UConversationInstance* StartConversation(const FGameplayTag& ConversationEntryTag, AActor* Instigator, const FGameplayTag& InstigatorTag,
+		AActor* Target, const FGameplayTag& TargetTag, const TSubclassOf<UConversationInstance> ConversationInstanceClass = nullptr);
+
+	UFUNCTION(BlueprintCallable, BlueprintAuthorityOnly, Category = "Conversation", meta = (DeterminesOutputType = "ConversationType"))
+	static UConversationInstance* StartConversationFromGraph(const FGameplayTag& ConversationEntryTag, AActor* Instigator, const FGameplayTag& InstigatorTag,
+		AActor* Target, const FGameplayTag& TargetTag, const UConversationDatabase* Graph);
 };
 
 #if UE_ENABLE_INCLUDE_ORDER_DEPRECATED_IN_5_2
diff --git a/Engine/Plugins/Experimental/CommonConversation/Source/CommonConversationRuntime/Public/ConversationRegistry.h b/Engine/Plugins/Experimental/CommonConversation/Source/CommonConversationRuntime/Public/ConversationRegistry.h
--- a/Engine/Plugins/Experimental/CommonConversation/Source/CommonConversationRuntime/Public/ConversationRegistry.h
+++ b/Engine/Plugins/Experimental/CommonConversation/Source/CommonConversationRuntime/Public/ConversationRegistry.h
@@ -83,11 +83,12 @@ public:
 	/** UWorldSubsystem End */
 
 	UConversationNode* GetRuntimeNodeFromGUID(const FGuid& NodeGUID) const;
-	TArray<FGuid> GetEntryPointGUIDs(FGameplayTag EntryPoint) const;
+	TArray<FGuid> GetEntryPointGUIDs(const FGameplayTag& EntryPoint) const;
 
-	TArray<FGuid> GetOutputLinkGUIDs(FGameplayTag EntryPoint) const;
+	TArray<FGuid> GetOutputLinkGUIDs(const FGameplayTag& EntryPoint) const;
 	TArray<FGuid> GetOutputLinkGUIDs(const FGuid& SourceGUID) const;
 	TArray<FGuid> GetOutputLinkGUIDs(const TArray<FGuid>& SourceGUIDs) const;
+	TArray<FGuid> GetOutputLinkGUIDs(const UConversationDatabase* Graph, const FGameplayTag& EntryPoint) const;
 
 	TSharedPtr<FConversationsHandle> LoadConversationsFor(const FGameplayTag& ConversationEntryTag) const;
 	TSharedPtr<FConversationsHandle> LoadConversationsFor(const TArray<FGameplayTag>& ConversationEntryTags) const;
