Updates CommonConversation plugin to filter node list with FBlueprintActionDatabase::IsClassAllowed.

diff --git a/Engine/Plugins/Experimental/CommonConversation/Source/CommonConversationGraph/CommonConversationGraph.Build.cs b/Engine/Plugins/Experimental/CommonConversation/Source/CommonConversationGraph/CommonConversationGraph.Build.cs
--- a/Engine/Plugins/Experimental/CommonConversation/Source/CommonConversationGraph/CommonConversationGraph.Build.cs
+++ b/Engine/Plugins/Experimental/CommonConversation/Source/CommonConversationGraph/CommonConversationGraph.Build.cs
@@ -18,6 +18,7 @@ public class CommonConversationGraph : ModuleRules
 		PrivateDependencyModuleNames.AddRange(
 			new string[]
 			{
+				"BlueprintGraph",
 				"CoreUObject",
 				"Engine",
 				"Slate",
diff --git a/Engine/Plugins/Experimental/CommonConversation/Source/CommonConversationGraph/Private/ConversationGraphSchema.cpp b/Engine/Plugins/Experimental/CommonConversation/Source/CommonConversationGraph/Private/ConversationGraphSchema.cpp
--- a/Engine/Plugins/Experimental/CommonConversation/Source/CommonConversationGraph/Private/ConversationGraphSchema.cpp
+++ b/Engine/Plugins/Experimental/CommonConversation/Source/CommonConversationGraph/Private/ConversationGraphSchema.cpp
@@ -20,6 +20,7 @@
 #include "ConversationGraphNode_Choice.h"
 #include "ConversationGraphNode_Knot.h"
 
+#include "BlueprintActionDatabase.h"
 #include "EdGraph/EdGraph.h"
 #include "GraphEditorActions.h"
 #include "ToolMenu.h"
@@ -114,15 +115,30 @@ void UConversationGraphSchema::AddConversationNodeOptions(const FString& Categor
 	TArray<FGraphNodeClassData> NodeClasses;
 	GetConversationClassCache().GatherClasses(RuntimeNodeType, /*out*/ NodeClasses);
 
-	for (const FGraphNodeClassData& NodeClass : NodeClasses)
+	for (FGraphNodeClassData& NodeClass : NodeClasses)
 	{
-		const FText NodeTypeName = FText::FromString(FName::NameToDisplayString(NodeClass.ToString(), false));
+		bool bIsAllowed = false;
+		// We check the package name only first to test the allowed status without possibly loading a full uasset class from disk
+		// If there is no package name, fallback to testing with a fully loaded class
+		if (!NodeClass.GetPackageName().IsEmpty())
+		{
+			bIsAllowed = FBlueprintActionDatabase::IsClassAllowed(FTopLevelAssetPath(NodeClass.GetPackageName()), FBlueprintActionDatabase::EPermissionsContext::Node);
+		}
+		else
+		{
+			bIsAllowed = FBlueprintActionDatabase::IsClassAllowed(NodeClass.GetClass(), FBlueprintActionDatabase::EPermissionsContext::Node);
+		}
+		
+		if (bIsAllowed)
+		{
+			const FText NodeTypeName = FText::FromString(FName::NameToDisplayString(NodeClass.ToString(), false));
 
-		TSharedPtr<FAISchemaAction_NewNode> AddOpAction = UAIGraphSchema::AddNewNodeAction(ListBuilder, NodeClass.GetCategory(), NodeTypeName, FText::GetEmpty());
+			TSharedPtr<FAISchemaAction_NewNode> AddOpAction = UAIGraphSchema::AddNewNodeAction(ListBuilder, NodeClass.GetCategory(), NodeTypeName, FText::GetEmpty());
 
-		UConversationGraphNode* OpNode = NewObject<UConversationGraphNode>(ContextMenuBuilder.OwnerOfTemporaries, EditorNodeType);
-		OpNode->ClassData = NodeClass;
-		AddOpAction->NodeTemplate = OpNode;
+			UConversationGraphNode* OpNode = NewObject<UConversationGraphNode>(ContextMenuBuilder.OwnerOfTemporaries, EditorNodeType);
+			OpNode->ClassData = NodeClass;
+			AddOpAction->NodeTemplate = OpNode;
+		}
 	}
 
 	ContextMenuBuilder.Append(ListBuilder);
