Updates CommonConversation plugin to filter subnode list with FBlueprintActionDatabase::IsClassAllowed.
Add flag to remove ChoiceClassName from the node title if requested.
Fix bug with conversation functions not looking for nodes in passed in graph.

diff --git a/Engine/Plugins/Experimental/CommonConversation/Source/CommonConversationGraph/Private/ConversationGraphNode_Choice.cpp b/Engine/Plugins/Experimental/CommonConversation/Source/CommonConversationGraph/Private/ConversationGraphNode_Choice.cpp
--- a/Engine/Plugins/Experimental/CommonConversation/Source/CommonConversationGraph/Private/ConversationGraphNode_Choice.cpp
+++ b/Engine/Plugins/Experimental/CommonConversation/Source/CommonConversationGraph/Private/ConversationGraphNode_Choice.cpp
@@ -41,7 +41,7 @@ FText UConversationGraphNode_Choice::GetNodeTitle(ENodeTitleType::Type TitleType
 			Args.Add(TEXT("ChoiceClassName"), FText::FromName(RuntimeNode->GetClass()->GetDisplayNameText()));
 #endif
 
-			const FText FormatString = bIsDerivedType ?
+			const FText FormatString = (bIsDerivedType && !RuntimeNode->GetHideChoiceClassName()) ?
 				LOCTEXT("UConversationGraphNode_Choice_GetNodeTitle_WithClass", "Choice \"{ChoiceName}\" ({ChoiceClassName})") :
 				LOCTEXT("UConversationGraphNode_Choice_GetNodeTitle", "Choice \"{ChoiceName}\"");
 
diff --git a/Engine/Plugins/Experimental/CommonConversation/Source/CommonConversationGraph/Private/ConversationGraphSchema.cpp b/Engine/Plugins/Experimental/CommonConversation/Source/CommonConversationGraph/Private/ConversationGraphSchema.cpp
--- a/Engine/Plugins/Experimental/CommonConversation/Source/CommonConversationGraph/Private/ConversationGraphSchema.cpp
+++ b/Engine/Plugins/Experimental/CommonConversation/Source/CommonConversationGraph/Private/ConversationGraphSchema.cpp
@@ -88,24 +88,45 @@ void UConversationGraphSchema::GetGraphNodeContextActions(FGraphContextMenuBuild
 void UConversationGraphSchema::GetSubNodeClasses(int32 SubNodeFlags, TArray<FGraphNodeClassData>& ClassData, UClass*& GraphNodeClass) const
 {
 	FGraphNodeClassHelper& ClassCache = GetConversationClassCache();
+	TArray<FGraphNodeClassData> TempClassData;
 
 	switch ((EConversationGraphSubNodeType)SubNodeFlags)
 	{
 	case EConversationGraphSubNodeType::Requirement:
-		ClassCache.GatherClasses(UConversationRequirementNode::StaticClass(), /*out*/ ClassData);
+		ClassCache.GatherClasses(UConversationRequirementNode::StaticClass(), /*out*/ TempClassData);
 		GraphNodeClass = UConversationGraphNode_Requirement::StaticClass();
 		break;
 	case EConversationGraphSubNodeType::SideEffect:
-		ClassCache.GatherClasses(UConversationSideEffectNode::StaticClass(), /*out*/ ClassData);
+		ClassCache.GatherClasses(UConversationSideEffectNode::StaticClass(), /*out*/ TempClassData);
 		GraphNodeClass = UConversationGraphNode_SideEffect::StaticClass();
 		break;
 	case EConversationGraphSubNodeType::Choice:
-		ClassCache.GatherClasses(UConversationChoiceNode::StaticClass(), /*out*/ ClassData);
+		ClassCache.GatherClasses(UConversationChoiceNode::StaticClass(), /*out*/ TempClassData);
 		GraphNodeClass = UConversationGraphNode_Choice::StaticClass();
 		break;
 	default:
 		unimplemented();
 	}
+
+	for (FGraphNodeClassData& Class : TempClassData)
+	{
+		bool bIsAllowed = false;
+		// We check the name only first to test the allowed status without possibly loading a full uasset class from disk
+		// If there is no package name, fallback to testing with a fully loaded class
+		if (!Class.GetPackageName().IsEmpty())
+		{
+			bIsAllowed = FBlueprintActionDatabase::IsClassAllowed(FTopLevelAssetPath(FName(Class.GetPackageName()), FName(Class.GetClassName())), FBlueprintActionDatabase::EPermissionsContext::Node);
+		}
+		else
+		{
+			bIsAllowed = FBlueprintActionDatabase::IsClassAllowed(Class.GetClass(), FBlueprintActionDatabase::EPermissionsContext::Node);
+		}
+
+		if (bIsAllowed)
+		{
+			ClassData.Add(std::move(Class));
+		}
+	}
 }
 
 void UConversationGraphSchema::AddConversationNodeOptions(const FString& CategoryName, FGraphContextMenuBuilder& ContextMenuBuilder, TSubclassOf<UConversationNode> RuntimeNodeType, TSubclassOf<UConversationGraphNode> EditorNodeType) const
@@ -118,11 +139,11 @@ void UConversationGraphSchema::AddConversationNodeOptions(const FString& Categor
 	for (FGraphNodeClassData& NodeClass : NodeClasses)
 	{
 		bool bIsAllowed = false;
-		// We check the package name only first to test the allowed status without possibly loading a full uasset class from disk
+		// We check the name only first to test the allowed status without possibly loading a full uasset class from disk
 		// If there is no package name, fallback to testing with a fully loaded class
 		if (!NodeClass.GetPackageName().IsEmpty())
 		{
-			bIsAllowed = FBlueprintActionDatabase::IsClassAllowed(FTopLevelAssetPath(NodeClass.GetPackageName()), FBlueprintActionDatabase::EPermissionsContext::Node);
+			bIsAllowed = FBlueprintActionDatabase::IsClassAllowed(FTopLevelAssetPath(FName(NodeClass.GetPackageName()), FName(NodeClass.GetClassName())), FBlueprintActionDatabase::EPermissionsContext::Node);
 		}
 		else
 		{
diff --git a/Engine/Plugins/Experimental/CommonConversation/Source/CommonConversationRuntime/Private/ConversationInstance.cpp b/Engine/Plugins/Experimental/CommonConversation/Source/CommonConversationRuntime/Private/ConversationInstance.cpp
--- a/Engine/Plugins/Experimental/CommonConversation/Source/CommonConversationRuntime/Private/ConversationInstance.cpp
+++ b/Engine/Plugins/Experimental/CommonConversation/Source/CommonConversationRuntime/Private/ConversationInstance.cpp
@@ -80,6 +80,7 @@ void UConversationInstance::ServerStartConversation(const FGameplayTag& EntryPoi
 
 	ResetConversationProgress();
 	StartingEntryGameplayTag = EntryPoint;
+	ActiveConversationGraph = Graph;
 
 	UConversationRegistry* ConversationRegistry = UConversationRegistry::GetFromWorld(GetWorld());
 	
@@ -502,7 +503,7 @@ TArray<FGuid> UConversationInstance::DetermineBranches(const TArray<FGuid>& Sour
 	TArray<FGuid> EnabledPaths;
 	for (const FGuid& TestGUID : SourceList)
 	{
-		UConversationNode* TestNode = Context.GetConversationRegistry().GetRuntimeNodeFromGUID(TestGUID);
+		UConversationNode* TestNode = Context.GetConversationRegistry().GetRuntimeNodeFromGUID(TestGUID, ActiveConversationGraph.Get());
 		if (UConversationTaskNode* TaskNode = Cast<UConversationTaskNode>(TestNode))
 		{
 			const EConversationRequirementResult RequirementResult = TaskNode->CheckRequirements(Context);
diff --git a/Engine/Plugins/Experimental/CommonConversation/Source/CommonConversationRuntime/Private/ConversationParticipantComponent.cpp b/Engine/Plugins/Experimental/CommonConversation/Source/CommonConversationRuntime/Private/ConversationParticipantComponent.cpp
--- a/Engine/Plugins/Experimental/CommonConversation/Source/CommonConversationRuntime/Private/ConversationParticipantComponent.cpp
+++ b/Engine/Plugins/Experimental/CommonConversation/Source/CommonConversationRuntime/Private/ConversationParticipantComponent.cpp
@@ -98,11 +98,11 @@ void UConversationParticipantComponent::ServerNotifyConversationEnded(UConversat
 	}
 }
 
-void UConversationParticipantComponent::ServerNotifyExecuteTaskAndSideEffects(const FConversationNodeHandle& Handle)
+void UConversationParticipantComponent::ServerNotifyExecuteTaskAndSideEffects(const FConversationNodeHandle& Handle, const UConversationDatabase* Graph)
 {
 	if (GetOwner()->GetLocalRole() == ROLE_Authority)
 	{
-		ClientExecuteTaskAndSideEffects(Handle);
+		ClientExecuteTaskAndSideEffects(Handle, Graph);
 	}
 }
 
@@ -321,9 +321,9 @@ void UConversationParticipantComponent::ClientUpdateParticipants_Implementation(
 	LastMessage.Participants = InParticipants;
 }
 
-void UConversationParticipantComponent::ClientExecuteTaskAndSideEffects_Implementation(FConversationNodeHandle Handle)
+void UConversationParticipantComponent::ClientExecuteTaskAndSideEffects_Implementation(FConversationNodeHandle Handle, const UConversationDatabase* Graph)
 {
-	if (const UConversationTaskNode* TaskNode = Cast<UConversationTaskNode>(Handle.TryToResolve_Slow(GetWorld())))
+	if (const UConversationTaskNode* TaskNode = Cast<UConversationTaskNode>(Handle.TryToResolve_Slow(GetWorld(), Graph)))
 	{
 		FConversationContext ClientContext = FConversationContext::CreateClientContext(this, TaskNode);
 		TaskNode->ExecuteTaskNodeWithSideEffects(ClientContext);
diff --git a/Engine/Plugins/Experimental/CommonConversation/Source/CommonConversationRuntime/Private/ConversationRegistry.cpp b/Engine/Plugins/Experimental/CommonConversation/Source/CommonConversationRuntime/Private/ConversationRegistry.cpp
--- a/Engine/Plugins/Experimental/CommonConversation/Source/CommonConversationRuntime/Private/ConversationRegistry.cpp
+++ b/Engine/Plugins/Experimental/CommonConversation/Source/CommonConversationRuntime/Private/ConversationRegistry.cpp
@@ -7,6 +7,7 @@
 #include "Engine/AssetManager.h"
 #include "UObject/UObjectIterator.h"
 #include "ConversationContext.h"
+#include "ConversationInstance.h"
 #include "Engine/StreamableManager.h"
 #include "GameFeaturesSubsystem.h"
 
@@ -129,15 +130,22 @@ bool FNetSerializeScriptStructCache_ConvVersion::NetSerialize(FArchive& Ar, UScr
 
 //======================================================================================
 
-const UConversationNode* FConversationNodeHandle::TryToResolve_Slow(UWorld* InWorld) const
+const UConversationNode* FConversationNodeHandle::TryToResolve_Slow(UWorld* InWorld, const UConversationDatabase* Graph) const
 {
-	UConversationRegistry* Registry = UConversationRegistry::GetFromWorld(InWorld);
-	return Registry->GetRuntimeNodeFromGUID(NodeGUID);
+	if (UConversationRegistry* Registry = UConversationRegistry::GetFromWorld(InWorld))
+	{
+		return Registry->TryGetRuntimeNodeFromGUID(NodeGUID, Graph);
+	}
+	return nullptr;
 }
 
 const UConversationNode* FConversationNodeHandle::TryToResolve(const FConversationContext& Context) const
 {
-	return Context.GetConversationRegistry().GetRuntimeNodeFromGUID(NodeGUID);
+	if (UConversationInstance* Instance = Context.TryGetActiveConversation())
+	{
+		return Context.GetConversationRegistry().TryGetRuntimeNodeFromGUID(NodeGUID, Instance->GetActiveConversationGraph());
+	}
+	return Context.GetConversationRegistry().TryGetRuntimeNodeFromGUID(NodeGUID);
 }
 
 UConversationRegistry::UConversationRegistry()
@@ -187,7 +195,21 @@ void UConversationRegistry::OnGameFeatureDeactivating(const UGameFeatureData* Ga
 	GameFeatureStateModified();
 }
 
-UConversationNode* UConversationRegistry::GetRuntimeNodeFromGUID(const FGuid& NodeGUID) const
+UConversationNode* UConversationRegistry::GetRuntimeNodeFromGUID(const FGuid& NodeGUID, const UConversationDatabase* Graph) const
+{
+	// It's possible this is just a null/empty guid, if that happens just return null.
+	if (!NodeGUID.IsValid())
+	{
+		UE_LOG(LogCommonConversationRuntime, Verbose, TEXT("UConversationRegistry::GetRuntimeNodeFromGUID - NodeGUID not valid"));
+		return nullptr;
+	}
+
+	UConversationNode* Node = TryGetRuntimeNodeFromGUID(NodeGUID, Graph);
+	ensureMsgf(Node != nullptr, TEXT("Unexpected GetRuntimeNodeFromGUID(%s) Failed. Nodes Searched: %d"), *NodeGUID.ToString(), NodeGuidToConversation.Num());
+	return Node;
+}
+
+UConversationNode* UConversationRegistry::TryGetRuntimeNodeFromGUID(const FGuid& NodeGUID, const UConversationDatabase* Graph) const
 {
 	UE_LOG(LogCommonConversationRuntime, Verbose, TEXT("Start UConversationRegistry::GetRuntimeNodeFromGUID with NodeGUID: (%s)"), *NodeGUID.ToString());
 
@@ -196,17 +218,14 @@ UConversationNode* UConversationRegistry::GetRuntimeNodeFromGUID(const FGuid& No
 	// It's possible this is just a null/empty guid, if that happens just return null.
 	if (!NodeGUID.IsValid())
 	{
-		UE_LOG(LogCommonConversationRuntime, Verbose, TEXT("UConversationRegistry::GetRuntimeNodeFromGUID - NodeGUID not valid"));
-
 		return nullptr;
 	}
 
-	if (const UConversationDatabase* ConversationDB = GetConversationFromNodeGUID(NodeGUID))
+	if (const UConversationDatabase* ConversationDB = Graph ? Graph : GetConversationFromNodeGUID(NodeGUID))
 	{
 		return ConversationDB->ReachableNodeMap.FindRef(NodeGUID);
 	}
 
-	ensureMsgf(false, TEXT("Unexpected GetRuntimeNodeFromGUID(%s) Failed. Nodes Searched: %d"), *NodeGUID.ToString(), NodeGuidToConversation.Num());
 	return nullptr;
 }
 
@@ -283,6 +302,29 @@ TArray<FGuid> UConversationRegistry::GetOutputLinkGUIDs(const UConversationDatab
 	return Result;
 }
 
+TArray<FGuid> UConversationRegistry::GetOutputLinkGUIDs(const UConversationDatabase* Graph, const FGuid& SourceGUID) const
+{
+	if (Graph == nullptr)
+	{
+		return GetOutputLinkGUIDs(SourceGUID);
+	}
+
+	const_cast<UConversationRegistry*>(this)->BuildDependenciesGraph();
+
+	TArray<FGuid> Result;
+
+	UConversationNode* SourceNode = Graph->ReachableNodeMap.FindRef(SourceGUID);
+	if (SourceNode)
+	{
+		if (UConversationNodeWithLinks* SourceNodeWithLinks = CastChecked<UConversationNodeWithLinks>(SourceNode))
+		{
+			Result.Append(SourceNodeWithLinks->OutputConnections);
+		}
+	}
+
+	return Result;
+}
+
 UConversationDatabase* UConversationRegistry::GetConversationFromNodeGUID(const FGuid& NodeGUID) const
 {
 	if (const FSoftObjectPath* ConversationPathPtr = NodeGuidToConversation.Find(NodeGUID))
diff --git a/Engine/Plugins/Experimental/CommonConversation/Source/CommonConversationRuntime/Private/ConversationTaskNode.cpp b/Engine/Plugins/Experimental/CommonConversation/Source/CommonConversationRuntime/Private/ConversationTaskNode.cpp
--- a/Engine/Plugins/Experimental/CommonConversation/Source/CommonConversationRuntime/Private/ConversationTaskNode.cpp
+++ b/Engine/Plugins/Experimental/CommonConversation/Source/CommonConversationRuntime/Private/ConversationTaskNode.cpp
@@ -127,13 +127,19 @@ FConversationTaskResult UConversationTaskNode::ExecuteTaskNodeWithSideEffects(co
 				}
 			}
 
+			const UConversationDatabase* Graph = nullptr;
+			if (const UConversationInstance* Conversation = InContext.GetActiveConversation())
+			{
+				Graph = Conversation->GetActiveConversationGraph();
+			}
+
 			FConversationParticipants Participants = InContext.GetParticipantsCopy();
 			for (const FConversationParticipantEntry& ParticipantEntry : Participants.List)
 			{
 				if (UConversationParticipantComponent* Component = ParticipantEntry.GetParticipantComponent())
 				{
 					// Notify each client in the conversation
-					Component->ServerNotifyExecuteTaskAndSideEffects(InContext.GetCurrentNodeHandle());
+					Component->ServerNotifyExecuteTaskAndSideEffects(InContext.GetCurrentNodeHandle(), Graph);
 				}
 			}
 		}
@@ -161,10 +167,15 @@ void UConversationTaskNode::GenerateChoicesForDestinations(FConversationBranchPo
 	check(InContext.IsServerContext());
 
 	UWorld* World = InContext.GetWorld();
+	const UConversationDatabase* Graph = nullptr;
+	if (const UConversationInstance* Conversation = InContext.GetActiveConversation())
+	{
+		Graph = Conversation->GetActiveConversationGraph();
+	}
 
 	for (const FGuid& DestinationGUID : CandidateDestinations)
 	{
-		if (UConversationTaskNode* DestinationTaskNode = Cast<UConversationTaskNode>(InContext.GetConversationRegistry().GetRuntimeNodeFromGUID(DestinationGUID)))
+		if (UConversationTaskNode* DestinationTaskNode = Cast<UConversationTaskNode>(InContext.GetConversationRegistry().GetRuntimeNodeFromGUID(DestinationGUID, Graph)))
 		{
 			TGuardValue<decltype(DestinationTaskNode->EvalWorldContextObj)> Swapper(DestinationTaskNode->EvalWorldContextObj, World);
 
@@ -197,7 +208,7 @@ void UConversationTaskNode::GenerateNextChoices(FConversationBranchPointBuilder&
 {
 	if (const UConversationInstance* Conversation = Context.GetActiveConversation())
 	{
-		const TArray<FGuid> CandidateDestinations = Context.GetConversationRegistry().GetOutputLinkGUIDs({ Conversation->GetCurrentChoiceReference().NodeReference.NodeGUID });
+		const TArray<FGuid> CandidateDestinations = Context.GetConversationRegistry().GetOutputLinkGUIDs(Conversation->GetActiveConversationGraph(), { Conversation->GetCurrentChoiceReference().NodeReference.NodeGUID });
 
 		UConversationTaskNode::GenerateChoicesForDestinations(BranchBuilder, Context, CandidateDestinations);
 	}
diff --git a/Engine/Plugins/Experimental/CommonConversation/Source/CommonConversationRuntime/Public/ConversationChoiceNode.h b/Engine/Plugins/Experimental/CommonConversation/Source/CommonConversationRuntime/Public/ConversationChoiceNode.h
--- a/Engine/Plugins/Experimental/CommonConversation/Source/CommonConversationRuntime/Public/ConversationChoiceNode.h
+++ b/Engine/Plugins/Experimental/CommonConversation/Source/CommonConversationRuntime/Public/ConversationChoiceNode.h
@@ -25,6 +25,8 @@ public:
 	UPROPERTY(BlueprintReadOnly, EditAnywhere, Category=Conversation)
 	FGameplayTagContainer ChoiceTags;
 
+	bool GetHideChoiceClassName() const { return bHideChoiceClassName; }
+
 	virtual bool GenerateChoice(const FConversationContext& Context, FClientConversationOptionEntry& ChoiceEntry) const;
 
 	virtual void NotifyChoicePickedByUser(const FConversationContext& InContext, const FClientConversationOptionEntry& InClientChoice) const;
@@ -32,6 +34,8 @@ public:
 protected:
 	UFUNCTION(BlueprintNativeEvent)
 	void FillChoice(const FConversationContext& Context, FClientConversationOptionEntry& ChoiceEntry) const;
+
+	bool bHideChoiceClassName = false;
 };
 
 #if UE_ENABLE_INCLUDE_ORDER_DEPRECATED_IN_5_2
diff --git a/Engine/Plugins/Experimental/CommonConversation/Source/CommonConversationRuntime/Public/ConversationContext.h b/Engine/Plugins/Experimental/CommonConversation/Source/CommonConversationRuntime/Public/ConversationContext.h
--- a/Engine/Plugins/Experimental/CommonConversation/Source/CommonConversationRuntime/Public/ConversationContext.h
+++ b/Engine/Plugins/Experimental/CommonConversation/Source/CommonConversationRuntime/Public/ConversationContext.h
@@ -178,6 +178,7 @@ public:
 	UConversationRegistry& GetConversationRegistry() const { check(ConversationRegistry) return *ConversationRegistry; }
 
 	UConversationInstance* GetActiveConversation() const { ensure(IsServerContext()); return ActiveConversation; }
+	UConversationInstance* TryGetActiveConversation() const { return IsServerContext() ? ActiveConversation : nullptr; }
 
 	const UConversationTaskNode* GetTaskBeingConsidered() const { return TaskBeingConsidered; }
 
diff --git a/Engine/Plugins/Experimental/CommonConversation/Source/CommonConversationRuntime/Public/ConversationInstance.h b/Engine/Plugins/Experimental/CommonConversation/Source/CommonConversationRuntime/Public/ConversationInstance.h
--- a/Engine/Plugins/Experimental/CommonConversation/Source/CommonConversationRuntime/Public/ConversationInstance.h
+++ b/Engine/Plugins/Experimental/CommonConversation/Source/CommonConversationRuntime/Public/ConversationInstance.h
@@ -87,6 +87,11 @@ public:
 		return Participants.GetParticipantComponent(ParticipantID);
 	}
 
+	const UConversationDatabase* GetActiveConversationGraph() const
+	{
+		return ActiveConversationGraph.Get();
+	}
+
 	const FConversationNodeHandle& GetCurrentNodeHandle() const { return CurrentBranchPoint.GetNodeHandle(); }
 	const FConversationChoiceReference& GetCurrentChoiceReference() const { return CurrentBranchPoint.ClientChoice.ChoiceReference; }
 	const TArray<FClientConversationOptionEntry>& GetCurrentUserConversationChoices() const { return CurrentUserChoices; }
@@ -128,6 +133,9 @@ private:
 	UPROPERTY()
 	FConversationParticipants Participants;
 
+	UPROPERTY()
+	TObjectPtr<const UConversationDatabase> ActiveConversationGraph = nullptr;
+
 	FGameplayTag StartingEntryGameplayTag;
 	FConversationBranchPoint StartingBranchPoint;
 
diff --git a/Engine/Plugins/Experimental/CommonConversation/Source/CommonConversationRuntime/Public/ConversationNode.h b/Engine/Plugins/Experimental/CommonConversation/Source/CommonConversationRuntime/Public/ConversationNode.h
--- a/Engine/Plugins/Experimental/CommonConversation/Source/CommonConversationRuntime/Public/ConversationNode.h
+++ b/Engine/Plugins/Experimental/CommonConversation/Source/CommonConversationRuntime/Public/ConversationNode.h
@@ -58,7 +58,7 @@ public:
 	 * in memory.
 	 */
 	const UConversationNode* TryToResolve(const FConversationContext& Context) const;
-	const UConversationNode* TryToResolve_Slow(UWorld* InWorld) const;
+	const UConversationNode* TryToResolve_Slow(UWorld* InWorld, const UConversationDatabase* Graph = nullptr) const;
 };
 
 inline bool operator==(const FConversationNodeHandle& Lhs, const FConversationNodeHandle& Rhs) { return Lhs.NodeGUID == Rhs.NodeGUID; }
diff --git a/Engine/Plugins/Experimental/CommonConversation/Source/CommonConversationRuntime/Public/ConversationParticipantComponent.h b/Engine/Plugins/Experimental/CommonConversation/Source/CommonConversationRuntime/Public/ConversationParticipantComponent.h
--- a/Engine/Plugins/Experimental/CommonConversation/Source/CommonConversationRuntime/Public/ConversationParticipantComponent.h
+++ b/Engine/Plugins/Experimental/CommonConversation/Source/CommonConversationRuntime/Public/ConversationParticipantComponent.h
@@ -48,7 +48,7 @@ public:
 #if WITH_SERVER_CODE
 	void ServerNotifyConversationStarted(UConversationInstance* Conversation, FGameplayTag AsParticipant);
 	void ServerNotifyConversationEnded(UConversationInstance* Conversation, const FConversationParticipants& PreservedParticipants);
-	void ServerNotifyExecuteTaskAndSideEffects(const FConversationNodeHandle& Handle);
+	void ServerNotifyExecuteTaskAndSideEffects(const FConversationNodeHandle& Handle, const UConversationDatabase* Graph = nullptr);
 	void ServerForAllConversationsRefreshChoices(UConversationInstance* IgnoreConversation = nullptr);
 	void ServerForAllConversationsRefreshTaskChoiceData(const FConversationNodeHandle& Handle, UConversationInstance* IgnoreConversation /*= nullptr*/);
 
@@ -91,7 +91,7 @@ protected:
 	void ClientUpdateParticipants(const FConversationParticipants& InParticipants);
 
 	UFUNCTION(Client, Reliable)
-	void ClientExecuteTaskAndSideEffects(FConversationNodeHandle Handle);
+	void ClientExecuteTaskAndSideEffects(FConversationNodeHandle Handle, const UConversationDatabase* Graph = nullptr);
 
 	UFUNCTION(Client, Reliable)
 	void ClientUpdateConversation(const FClientConversationMessagePayload& Message);
diff --git a/Engine/Plugins/Experimental/CommonConversation/Source/CommonConversationRuntime/Public/ConversationRegistry.h b/Engine/Plugins/Experimental/CommonConversation/Source/CommonConversationRuntime/Public/ConversationRegistry.h
--- a/Engine/Plugins/Experimental/CommonConversation/Source/CommonConversationRuntime/Public/ConversationRegistry.h
+++ b/Engine/Plugins/Experimental/CommonConversation/Source/CommonConversationRuntime/Public/ConversationRegistry.h
@@ -82,13 +82,15 @@ public:
 	virtual void Deinitialize() override;
 	/** UWorldSubsystem End */
 
-	UConversationNode* GetRuntimeNodeFromGUID(const FGuid& NodeGUID) const;
+	UConversationNode* GetRuntimeNodeFromGUID(const FGuid& NodeGUID, const UConversationDatabase* Graph = nullptr) const;
+	UConversationNode* TryGetRuntimeNodeFromGUID(const FGuid& NodeGUID, const UConversationDatabase* Graph = nullptr) const;
 	TArray<FGuid> GetEntryPointGUIDs(const FGameplayTag& EntryPoint) const;
 
 	TArray<FGuid> GetOutputLinkGUIDs(const FGameplayTag& EntryPoint) const;
 	TArray<FGuid> GetOutputLinkGUIDs(const FGuid& SourceGUID) const;
 	TArray<FGuid> GetOutputLinkGUIDs(const TArray<FGuid>& SourceGUIDs) const;
 	TArray<FGuid> GetOutputLinkGUIDs(const UConversationDatabase* Graph, const FGameplayTag& EntryPoint) const;
+	TArray<FGuid> GetOutputLinkGUIDs(const UConversationDatabase* Graph, const FGuid& SourceGUID) const;
 
 	TSharedPtr<FConversationsHandle> LoadConversationsFor(const FGameplayTag& ConversationEntryTag) const;
 	TSharedPtr<FConversationsHandle> LoadConversationsFor(const TArray<FGameplayTag>& ConversationEntryTags) const;
