Conversation System: Added some new functionality to the conversation instance class to allow for easy re-processing of nodes without having to force changes in the conversation branch points

#tests 28.20 PIE

diff --git a/Engine/Plugins/Experimental/CommonConversation/Source/CommonConversationRuntime/Private/ConversationInstance.cpp b/Engine/Plugins/Experimental/CommonConversation/Source/CommonConversationRuntime/Private/ConversationInstance.cpp
--- a/Engine/Plugins/Experimental/CommonConversation/Source/CommonConversationRuntime/Private/ConversationInstance.cpp
+++ b/Engine/Plugins/Experimental/CommonConversation/Source/CommonConversationRuntime/Private/ConversationInstance.cpp
@@ -428,6 +428,12 @@ void UConversationInstance::ServerRefreshTaskChoiceData(const FConversationNodeH
 		}
 	}
 }
+
+void UConversationInstance::ServerRefreshCurrentConversationNode()
+{
+	ProcessCurrentConversationNode();
+}
+
 #endif // #if WITH_SERVER_CODE
 
 void UConversationInstance::ResetConversationProgress()
@@ -521,6 +527,11 @@ TArray<FGuid> UConversationInstance::DetermineBranches(const TArray<FGuid>& Sour
 }
 
 void UConversationInstance::OnCurrentConversationNodeModified()
+{
+	ProcessCurrentConversationNode();
+}
+
+void UConversationInstance::ProcessCurrentConversationNode()
 {
 	check(GetCurrentChoiceReference().IsValid());
 
@@ -587,4 +598,4 @@ void UConversationInstance::OnCurrentConversationNodeModified()
 	}
 }
 
-#endif
+#endif //WITH_SERVER_CODE
diff --git a/Engine/Plugins/Experimental/CommonConversation/Source/CommonConversationRuntime/Public/ConversationInstance.h b/Engine/Plugins/Experimental/CommonConversation/Source/CommonConversationRuntime/Public/ConversationInstance.h
--- a/Engine/Plugins/Experimental/CommonConversation/Source/CommonConversationRuntime/Public/ConversationInstance.h
+++ b/Engine/Plugins/Experimental/CommonConversation/Source/CommonConversationRuntime/Public/ConversationInstance.h
@@ -57,6 +57,10 @@ public:
 
 	void ServerRefreshTaskChoiceData(const FConversationNodeHandle& Handle);
 
+	/** Attempts to process the current conversation node again - only useful in very specific circumstances where you'd want to re-run the current node
+	  * without having to deal with conversation flow changes. */
+	void ServerRefreshCurrentConversationNode();
+
 	/**
      * This is memory that will last for the duration of the conversation instance.  Don't store
      * anything here you want to be long lived.
@@ -123,7 +127,9 @@ private:
 
 #if WITH_SERVER_CODE
 	void OnCurrentConversationNodeModified();
-#endif
+
+	void ProcessCurrentConversationNode();
+#endif //WITH_SERVER_CODE
 
 protected:
 
