Add logging to try and track down the cause of a confusing startup crash where it is initializing subsystems incorrectly on some platforms

diff --git a/Engine/Plugins/Experimental/CommonConversation/Source/CommonConversationRuntime/Private/ConversationRegistry.cpp b/Engine/Plugins/Experimental/CommonConversation/Source/CommonConversationRuntime/Private/ConversationRegistry.cpp
--- a/Engine/Plugins/Experimental/CommonConversation/Source/CommonConversationRuntime/Private/ConversationRegistry.cpp
+++ b/Engine/Plugins/Experimental/CommonConversation/Source/CommonConversationRuntime/Private/ConversationRegistry.cpp
@@ -160,9 +160,20 @@ UConversationRegistry* UConversationRegistry::GetFromWorld(const UWorld* World)
 
 void UConversationRegistry::Initialize(FSubsystemCollectionBase& Collection)
 {
+	UE_LOG(LogCommonConversationRuntime, Log, TEXT("Initializing UConversationRegistry %s"), *GetPathNameSafe(this));
+
 	Super::Initialize(Collection);
 
-	UGameFeaturesSubsystem::Get().AddObserver(this);
+	if (UGameFeaturesSubsystem* GameFeaturesSubsystem = GEngine ? GEngine->GetEngineSubsystem<UGameFeaturesSubsystem>() : nullptr)
+	{
+		GameFeaturesSubsystem->AddObserver(this);
+	}
+	else
+	{
+		// Logging to track down a crash
+		UClass* ParentClass = this && this->GetClass() ? this->GetClass()->GetSuperClass() : nullptr;
+		UE_LOG(LogCommonConversationRuntime, Error, TEXT("GameFeaturesSubsystem missing during UConversationRegistry::Initialize! Collection class is %s and parent is %s"), *GetPathNameSafe(Collection.GetBaseType()), *GetPathNameSafe(ParentClass));
+	}
 }
 
 void UConversationRegistry::Deinitialize()
diff --git a/Engine/Source/Runtime/Engine/Private/Subsystems/SubsystemCollection.cpp b/Engine/Source/Runtime/Engine/Private/Subsystems/SubsystemCollection.cpp
--- a/Engine/Source/Runtime/Engine/Private/Subsystems/SubsystemCollection.cpp
+++ b/Engine/Source/Runtime/Engine/Private/Subsystems/SubsystemCollection.cpp
@@ -123,6 +123,8 @@ void FSubsystemCollectionBase::Initialize(UObject* NewOuter)
 		{
 			FSubsystemModuleWatcher::InitializeModuleWatcher();
 		}
+
+		UE_LOG(LogSubsystemCollection, Verbose, TEXT("Initializing subsystem collection for %s with type %s"), *GetNameSafe(NewOuter), *GetNameSafe(BaseType));
 		
 		TGuardValue<bool> PopulatingGuard(bPopulating, true);
 
