Fix CIS, static checker didn't like the cautious  null checks

diff --git a/Engine/Plugins/Experimental/CommonConversation/Source/CommonConversationRuntime/Private/ConversationRegistry.cpp b/Engine/Plugins/Experimental/CommonConversation/Source/CommonConversationRuntime/Private/ConversationRegistry.cpp
--- a/Engine/Plugins/Experimental/CommonConversation/Source/CommonConversationRuntime/Private/ConversationRegistry.cpp
+++ b/Engine/Plugins/Experimental/CommonConversation/Source/CommonConversationRuntime/Private/ConversationRegistry.cpp
@@ -171,7 +171,7 @@ void UConversationRegistry::Initialize(FSubsystemCollectionBase& Collection)
 	else
 	{
 		// Logging to track down a crash
-		UClass* ParentClass = this && this->GetClass() ? this->GetClass()->GetSuperClass() : nullptr;
+		UClass* ParentClass = GetClass() ? GetClass()->GetSuperClass() : nullptr;
 		UE_LOG(LogCommonConversationRuntime, Error, TEXT("GameFeaturesSubsystem missing during UConversationRegistry::Initialize! Collection class is %s and parent is %s"), *GetPathNameSafe(Collection.GetBaseType()), *GetPathNameSafe(ParentClass));
 	}
 }
