Partial backout of CL 31114688 to remove temporary logging as the problem was actually with a different subsystem

diff --git a/Engine/Plugins/Experimental/CommonConversation/Source/CommonConversationRuntime/Private/ConversationRegistry.cpp b/Engine/Plugins/Experimental/CommonConversation/Source/CommonConversationRuntime/Private/ConversationRegistry.cpp
--- a/Engine/Plugins/Experimental/CommonConversation/Source/CommonConversationRuntime/Private/ConversationRegistry.cpp
+++ b/Engine/Plugins/Experimental/CommonConversation/Source/CommonConversationRuntime/Private/ConversationRegistry.cpp
@@ -160,20 +160,9 @@ UConversationRegistry* UConversationRegistry::GetFromWorld(const UWorld* World)
 
 void UConversationRegistry::Initialize(FSubsystemCollectionBase& Collection)
 {
-	UE_LOG(LogCommonConversationRuntime, Log, TEXT("Initializing UConversationRegistry %s"), *GetPathNameSafe(this));
-
 	Super::Initialize(Collection);
 
-	if (UGameFeaturesSubsystem* GameFeaturesSubsystem = GEngine ? GEngine->GetEngineSubsystem<UGameFeaturesSubsystem>() : nullptr)
-	{
-		GameFeaturesSubsystem->AddObserver(this);
-	}
-	else
-	{
-		// Logging to track down a crash
-		UClass* ParentClass = GetClass() ? GetClass()->GetSuperClass() : nullptr;
-		UE_LOG(LogCommonConversationRuntime, Error, TEXT("GameFeaturesSubsystem missing during UConversationRegistry::Initialize! Collection class is %s and parent is %s"), *GetPathNameSafe(Collection.GetBaseType()), *GetPathNameSafe(ParentClass));
-	}
+	UGameFeaturesSubsystem::Get().AddObserver(this);
 }
 
 void UConversationRegistry::Deinitialize()
