Allow cycles on conversation graphs by CVar.
Hide "Add Subnode" buttons if the the sub-menu will be empty.

diff --git a/Engine/Plugins/Experimental/CommonConversation/Source/CommonConversationGraph/Private/ConversationGraphNode_Task.cpp b/Engine/Plugins/Experimental/CommonConversation/Source/CommonConversationGraph/Private/ConversationGraphNode_Task.cpp
--- a/Engine/Plugins/Experimental/CommonConversation/Source/CommonConversationGraph/Private/ConversationGraphNode_Task.cpp
+++ b/Engine/Plugins/Experimental/CommonConversation/Source/CommonConversationGraph/Private/ConversationGraphNode_Task.cpp
@@ -3,6 +3,8 @@
 #include "ConversationGraphNode_Task.h"
 #include "ConversationTaskNode.h"
 #include "ConversationGraphTypes.h"
+#include "ConversationGraphSchema.h"
+#include "EdGraph/EdGraph.h"
 #include "SGraphEditorActionMenuAI.h"
 #include "ToolMenu.h"
 #include "ToolMenuSection.h"
@@ -61,12 +63,16 @@ void UConversationGraphNode_Task::GetNodeContextMenuActions(UToolMenu* Menu, UGr
 	};
 
 	FToolMenuSection& Section = Menu->FindOrAddSection("ConversationGraphNode_Task");
+	const UConversationGraphSchema* Schema = Context->Graph ? Cast<const UConversationGraphSchema>(Context->Graph->GetSchema()) : nullptr;
 
-	Section.AddSubMenu(
-		"AddRequirement",
-		LOCTEXT("AddRequirement", "Add Requirement..."),
-		LOCTEXT("AddRequirementTooltip", "Adds a new requirement as a subnode"),
-		FNewMenuDelegate::CreateLambda(CreateAddRequirementSubMenu, (UEdGraph*)Context->Graph));
+	if (!Schema || Schema->HasSubNodeClasses(int32(EConversationGraphSubNodeType::Requirement)))
+	{
+		Section.AddSubMenu(
+			"AddRequirement",
+			LOCTEXT("AddRequirement", "Add Requirement..."),
+			LOCTEXT("AddRequirementTooltip", "Adds a new requirement as a subnode"),
+			FNewMenuDelegate::CreateLambda(CreateAddRequirementSubMenu, (UEdGraph*)Context->Graph));
+	}
 
 
 	auto CreateAddSideEffectSubMenu = [WeakThis](FMenuBuilder& MenuBuilder, UEdGraph* Graph)
@@ -81,11 +87,14 @@ void UConversationGraphNode_Task::GetNodeContextMenuActions(UToolMenu* Menu, UGr
 		MenuBuilder.AddWidget(Menu, FText(), true);
 	};
 
-	Section.AddSubMenu(
-		"AddSideEffect",
-		LOCTEXT("AddSideEffect", "Add Side Effect..."),
-		LOCTEXT("AddSideEffectTooltip", "Adds a new side effect as a subnode"),
-		FNewMenuDelegate::CreateLambda(CreateAddSideEffectSubMenu, (UEdGraph*)Context->Graph));
+	if (!Schema || Schema->HasSubNodeClasses(int32(EConversationGraphSubNodeType::SideEffect)))
+	{
+		Section.AddSubMenu(
+			"AddSideEffect",
+			LOCTEXT("AddSideEffect", "Add Side Effect..."),
+			LOCTEXT("AddSideEffectTooltip", "Adds a new side effect as a subnode"),
+			FNewMenuDelegate::CreateLambda(CreateAddSideEffectSubMenu, (UEdGraph*)Context->Graph));
+	}
 
 
 	auto CreateAddChoiceSubMenu = [WeakThis](FMenuBuilder& MenuBuilder, UEdGraph* Graph)
@@ -100,11 +109,14 @@ void UConversationGraphNode_Task::GetNodeContextMenuActions(UToolMenu* Menu, UGr
 		MenuBuilder.AddWidget(Menu, FText(), true);
 	};
 
-	Section.AddSubMenu(
-		"AddChoice",
-		LOCTEXT("AddChoice", "Add Choice..."),
-		LOCTEXT("AddChoiceTooltip", "Adds a new choice as a subnode"),
-		FNewMenuDelegate::CreateLambda(CreateAddChoiceSubMenu, (UEdGraph*)Context->Graph));
+	if (!Schema || Schema->HasSubNodeClasses(int32(EConversationGraphSubNodeType::Choice)))
+	{
+		Section.AddSubMenu(
+			"AddChoice",
+			LOCTEXT("AddChoice", "Add Choice..."),
+			LOCTEXT("AddChoiceTooltip", "Adds a new choice as a subnode"),
+			FNewMenuDelegate::CreateLambda(CreateAddChoiceSubMenu, (UEdGraph*)Context->Graph));
+	}
 }
 
 #undef LOCTEXT_NAMESPACE
diff --git a/Engine/Plugins/Experimental/CommonConversation/Source/CommonConversationGraph/Private/ConversationGraphSchema.cpp b/Engine/Plugins/Experimental/CommonConversation/Source/CommonConversationGraph/Private/ConversationGraphSchema.cpp
--- a/Engine/Plugins/Experimental/CommonConversation/Source/CommonConversationGraph/Private/ConversationGraphSchema.cpp
+++ b/Engine/Plugins/Experimental/CommonConversation/Source/CommonConversationGraph/Private/ConversationGraphSchema.cpp
@@ -31,6 +31,17 @@
 
 #define LOCTEXT_NAMESPACE "ConversationEditor"
 
+namespace ConversationEditorCVar
+{
+	static bool CheckForCyclesCVar = true;
+	FAutoConsoleVariableRef CVarCheckForCycles(
+		TEXT("ConversationEditor.CheckForCycles"),
+		CheckForCyclesCVar,
+		TEXT("This cvar controles if the Conversation Editor should check for cycles when links are created.\n")
+		TEXT("0: Don't Check, 1: Check for Cycles (Default)"),
+		ECVF_Default);
+}
+
 TSharedPtr<FGraphNodeClassHelper> ConversationClassCache;
 
 FGraphNodeClassHelper& GetConversationClassCache()
@@ -85,6 +96,14 @@ void UConversationGraphSchema::GetGraphNodeContextActions(FGraphContextMenuBuild
 	Super::GetGraphNodeContextActions(ContextMenuBuilder, SubNodeFlags);
 }
 
+bool UConversationGraphSchema::HasSubNodeClasses(int32 SubNodeFlags) const
+{
+	TArray<FGraphNodeClassData> TempClassData;
+	UClass* TempClass = nullptr;
+	GetSubNodeClasses(SubNodeFlags, TempClassData, TempClass);
+	return !TempClassData.IsEmpty();
+}
+
 void UConversationGraphSchema::GetSubNodeClasses(int32 SubNodeFlags, TArray<FGraphNodeClassData>& ClassData, UClass*& GraphNodeClass) const
 {
 	FGraphNodeClassHelper& ClassCache = GetConversationClassCache();
@@ -314,11 +333,14 @@ const FPinConnectionResponse UConversationGraphSchema::CanCreateConnection(const
 		TSet<UEdGraphNode*> VisitedNodes;
 	};
 
-	// check for cycles
-	FNodeVisitorCycleChecker CycleChecker;
-	if (!CycleChecker.CheckForLoop(PinA->GetOwningNode(), PinB->GetOwningNode()))
+	if (ConversationEditorCVar::CheckForCyclesCVar)
 	{
-		return FPinConnectionResponse(CONNECT_RESPONSE_DISALLOW, LOCTEXT("PinErrorcycle", "Can't create a graph cycle"));
+		// check for cycles
+		FNodeVisitorCycleChecker CycleChecker;
+		if (!CycleChecker.CheckForLoop(PinA->GetOwningNode(), PinB->GetOwningNode()))
+		{
+			return FPinConnectionResponse(CONNECT_RESPONSE_DISALLOW, LOCTEXT("PinErrorcycle", "Can't create a graph cycle"));
+		}
 	}
 
 	const bool bPinASingleLink = bPinAIsSingleComposite || bPinAIsSingleTask || bPinAIsSingleNode;
diff --git a/Engine/Plugins/Experimental/CommonConversation/Source/CommonConversationGraph/Public/ConversationGraphSchema.h b/Engine/Plugins/Experimental/CommonConversation/Source/CommonConversationGraph/Public/ConversationGraphSchema.h
--- a/Engine/Plugins/Experimental/CommonConversation/Source/CommonConversationGraph/Public/ConversationGraphSchema.h
+++ b/Engine/Plugins/Experimental/CommonConversation/Source/CommonConversationGraph/Public/ConversationGraphSchema.h
@@ -60,6 +60,7 @@ public:
 	//~ End EdGraphSchema Interface
 
 	virtual void GetGraphNodeContextActions(FGraphContextMenuBuilder& ContextMenuBuilder, int32 SubNodeFlags) const override;
+	bool HasSubNodeClasses(int32 SubNodeFlags) const;
 	virtual void GetSubNodeClasses(int32 SubNodeFlags, TArray<FGraphNodeClassData>& ClassData, UClass*& GraphNodeClass) const override;
 
 protected:
