Add Graph to TryToResolveChoice node.

diff --git a/Engine/Plugins/Experimental/CommonConversation/Source/CommonConversationRuntime/Public/ConversationTypes.h b/Engine/Plugins/Experimental/CommonConversation/Source/CommonConversationRuntime/Public/ConversationTypes.h
--- a/Engine/Plugins/Experimental/CommonConversation/Source/CommonConversationRuntime/Public/ConversationTypes.h
+++ b/Engine/Plugins/Experimental/CommonConversation/Source/CommonConversationRuntime/Public/ConversationTypes.h
@@ -8,6 +8,7 @@
 #include "ConversationTypes.generated.h"
 
 class UActorComponent;
+class UConversationDatabase;
 class UConversationInstance;
 class UConversationParticipantComponent;
 struct FConversationContext;
@@ -333,9 +334,9 @@ public:
 	}
 
 	template<class TConversationNodeClass = UConversationNode>
-	const TConversationNodeClass* TryToResolveChoiceNode_Slow(UWorld* InWorld) const
+	const TConversationNodeClass* TryToResolveChoiceNode_Slow(UWorld* InWorld, const UConversationDatabase* Graph = nullptr) const
 	{
-		return Cast<TConversationNodeClass>(ChoiceReference.NodeReference.TryToResolve_Slow(InWorld));
+		return Cast<TConversationNodeClass>(ChoiceReference.NodeReference.TryToResolve_Slow(InWorld, Graph));
 	}
 
 	void SetChoiceAvailable(bool bIsAvailable) { ChoiceType = bIsAvailable ? EConversationChoiceType::UserChoiceAvailable : EConversationChoiceType::UserChoiceUnavailable; }
