Several updates to Conversation Graph System.
   * Add function for getting a corresponding GraphNode from a ConversationNode
   * Add CVar for selecting the whole node when a subnode is dropped on it
   * Add function to retrieve the Add Requirement menu for a custom menu
   * Remove redundant IsUnique check when resetting detail?s panel

diff --git a/Engine/Plugins/Experimental/CommonConversation/Source/CommonConversationGraph/Private/ConversationGraphNode_Task.cpp b/Engine/Plugins/Experimental/CommonConversation/Source/CommonConversationGraph/Private/ConversationGraphNode_Task.cpp
--- a/Engine/Plugins/Experimental/CommonConversation/Source/CommonConversationGraph/Private/ConversationGraphNode_Task.cpp
+++ b/Engine/Plugins/Experimental/CommonConversation/Source/CommonConversationGraph/Private/ConversationGraphNode_Task.cpp
@@ -119,4 +119,18 @@ void UConversationGraphNode_Task::GetNodeContextMenuActions(UToolMenu* Menu, UGr
 	}
 }
 
+void UConversationGraphNode_Task::ApplyAddRequirementMenu(FMenuBuilder& MenuBuilder, UEdGraph* Graph)
+{
+	TWeakObjectPtr<UConversationGraphNode> WeakThis(const_cast<UConversationGraphNode_Task*>(this));
+
+	TSharedRef<SGraphEditorActionMenuAI> Menu =
+		SNew(SGraphEditorActionMenuAI)
+		.GraphObj(Graph)
+		.GraphNode(WeakThis.Get())
+		.SubNodeFlags((int32)EConversationGraphSubNodeType::Requirement)
+		.AutoExpandActionMenu(true);
+
+	MenuBuilder.AddWidget(Menu, FText(), true);
+}
+
 #undef LOCTEXT_NAMESPACE
diff --git a/Engine/Plugins/Experimental/CommonConversation/Source/CommonConversationGraph/Private/SConversationGraphNode.cpp b/Engine/Plugins/Experimental/CommonConversation/Source/CommonConversationGraph/Private/SConversationGraphNode.cpp
--- a/Engine/Plugins/Experimental/CommonConversation/Source/CommonConversationGraph/Private/SConversationGraphNode.cpp
+++ b/Engine/Plugins/Experimental/CommonConversation/Source/CommonConversationGraph/Private/SConversationGraphNode.cpp
@@ -42,6 +42,17 @@ namespace
 	static const bool bShowExecutionIndexInEditorMode = false;
 }
 
+namespace ConversationEditorCVar
+{
+	static bool SelectWholeNodeOnDropCVar = false;
+	FAutoConsoleVariableRef CVarSelectWholeNodeOnDrop(
+		TEXT("ConversationEditor.SelectWholeNodeOnDrop"),
+		SelectWholeNodeOnDropCVar,
+		TEXT("This cvar controles whether the subnode or whole node will be selected when a subnode is dropped onto the node.\n")
+		TEXT("0: SubNode will remain selected (Default), 1: Whole Node will be selected"),
+		ECVF_Default);
+}
+
 /////////////////////////////////////////////////////
 // SConversationPin
 
@@ -1086,6 +1097,18 @@ void SConversationGraphNode::MoveTo(const FVector2D& NewPosition, FNodeSet& Node
 	}
 }
 
+FReply SConversationGraphNode::OnDrop(const FGeometry& MyGeometry, const FDragDropEvent& DragDropEvent)
+{
+	FReply Reply = SGraphNodeAI::OnDrop(MyGeometry, DragDropEvent);
+
+	if (ConversationEditorCVar::SelectWholeNodeOnDropCVar)
+	{
+		GetOwnerPanel()->SelectionManager.SelectSingleNode(GraphNode);
+	}
+
+	return Reply;
+}
+
 bool SConversationGraphNode::IsNodeReachable() const
 {
 	UConversationGraphNode* ConversationGraphNode = CastChecked<UConversationGraphNode>(GraphNode);
diff --git a/Engine/Plugins/Experimental/CommonConversation/Source/CommonConversationGraph/Private/SConversationGraphNode.h b/Engine/Plugins/Experimental/CommonConversation/Source/CommonConversationGraph/Private/SConversationGraphNode.h
--- a/Engine/Plugins/Experimental/CommonConversation/Source/CommonConversationGraph/Private/SConversationGraphNode.h
+++ b/Engine/Plugins/Experimental/CommonConversation/Source/CommonConversationGraph/Private/SConversationGraphNode.h
@@ -49,6 +49,7 @@ public:
 	virtual TArray<FOverlayWidgetInfo> GetOverlayWidgets(bool bSelected, const FVector2D& WidgetSize) const override;
 	virtual TSharedRef<SGraphNode> GetNodeUnderMouse(const FGeometry& MyGeometry, const FPointerEvent& MouseEvent) override;
 	virtual void MoveTo(const FVector2D& NewPosition, FNodeSet& NodeFilter, bool bMarkDirty = true) override;
+	virtual FReply OnDrop(const FGeometry& MyGeometry, const FDragDropEvent& DragDropEvent) override;
 	// End of SGraphNode interface
 
 	/** handle double click */
diff --git a/Engine/Plugins/Experimental/CommonConversation/Source/CommonConversationGraph/Public/ConversationGraphNode_Task.h b/Engine/Plugins/Experimental/CommonConversation/Source/CommonConversationGraph/Public/ConversationGraphNode_Task.h
--- a/Engine/Plugins/Experimental/CommonConversation/Source/CommonConversationGraph/Public/ConversationGraphNode_Task.h
+++ b/Engine/Plugins/Experimental/CommonConversation/Source/CommonConversationGraph/Public/ConversationGraphNode_Task.h
@@ -11,6 +11,8 @@ class UToolMenu;
 class UEdGraph;
 class UEdGraphSchema;
 
+class FMenuBuilder;
+
 UCLASS()
 class COMMONCONVERSATIONGRAPH_API UConversationGraphNode_Task : public UConversationGraphNode
 {
@@ -22,6 +24,7 @@ class COMMONCONVERSATIONGRAPH_API UConversationGraphNode_Task : public UConversa
 
 	/** Gets a list of actions that can be done to this particular node */
 	virtual void GetNodeContextMenuActions(UToolMenu* Menu, UGraphNodeContextMenuContext* Context) const override;
+	virtual void ApplyAddRequirementMenu(FMenuBuilder& MenuBuilder, UEdGraph* Graph);
 
 	virtual bool CanPlaceBreakpoints() const override { return true; }
 };
diff --git a/Engine/Plugins/Experimental/CommonConversation/Source/CommonConversationRuntime/CommonConversationRuntime.Build.cs b/Engine/Plugins/Experimental/CommonConversation/Source/CommonConversationRuntime/CommonConversationRuntime.Build.cs
--- a/Engine/Plugins/Experimental/CommonConversation/Source/CommonConversationRuntime/CommonConversationRuntime.Build.cs
+++ b/Engine/Plugins/Experimental/CommonConversation/Source/CommonConversationRuntime/CommonConversationRuntime.Build.cs
@@ -24,6 +24,7 @@ public class CommonConversationRuntime : ModuleRules
 				"Slate",
 				"SlateCore",
 				"Projects",
+				"AIGraph",
 				"AIModule",
 				"NetCore",
 				"GameFeatures"
diff --git a/Engine/Plugins/Experimental/CommonConversation/Source/CommonConversationRuntime/Private/ConversationDatabase.cpp b/Engine/Plugins/Experimental/CommonConversation/Source/CommonConversationRuntime/Private/ConversationDatabase.cpp
--- a/Engine/Plugins/Experimental/CommonConversation/Source/CommonConversationRuntime/Private/ConversationDatabase.cpp
+++ b/Engine/Plugins/Experimental/CommonConversation/Source/CommonConversationRuntime/Private/ConversationDatabase.cpp
@@ -12,6 +12,10 @@
 #include "AssetRegistry/AssetRegistryModule.h"
 #include "AssetRegistry/AssetRegistryState.h"
 #include "UObject/AssetRegistryTagsContext.h"
+#if WITH_EDITORONLY_DATA
+#include "AIGraphNode.h"
+#include "ConversationNode.h"
+#endif
 
 #include UE_INLINE_GENERATED_CPP_BY_NAME(ConversationDatabase)
 
@@ -133,4 +137,24 @@ EDataValidationResult UConversationDatabase::IsDataValid(FDataValidationContext&
 }
 #endif // WITH_EDITOR
 
+#if WITH_EDITORONLY_DATA
+TObjectPtr<UEdGraphNode> UConversationDatabase::GetSourceGraphNode(const UConversationNode* NodeToFind) const
+{
+	for (const TObjectPtr<UEdGraph>& Graph : SourceGraphs)
+	{
+		for (TObjectPtr<UEdGraphNode>& Node : Graph->Nodes)
+		{
+			if (const UAIGraphNode* AIGraphNode = Cast<UAIGraphNode>(Node.Get()))
+			{
+				if (AIGraphNode->NodeInstance.Get() == NodeToFind)
+				{
+					return Node;
+				}
+			}
+		}
+	}
+	return TObjectPtr<UEdGraphNode>();
+}
+#endif // WITH_EDITORONLY_DATA
+
 #undef LOCTEXT_NAMESPACE
diff --git a/Engine/Plugins/Experimental/CommonConversation/Source/CommonConversationRuntime/Public/ConversationDatabase.h b/Engine/Plugins/Experimental/CommonConversation/Source/CommonConversationRuntime/Public/ConversationDatabase.h
--- a/Engine/Plugins/Experimental/CommonConversation/Source/CommonConversationRuntime/Public/ConversationDatabase.h
+++ b/Engine/Plugins/Experimental/CommonConversation/Source/CommonConversationRuntime/Public/ConversationDatabase.h
@@ -134,6 +134,8 @@ private:
 	TArray<TObjectPtr<UEdGraph>> SourceGraphs;
 
 public:
+	TObjectPtr<class UEdGraphNode> GetSourceGraphNode(const UConversationNode* NodeToFind) const;
+
 	// Info about the graphs we last edited
 	UPROPERTY()
 	TArray<FEditedDocumentInfo> LastEditedDocuments;
diff --git a/Engine/Source/Editor/PropertyEditor/Private/SDetailsViewBase.cpp b/Engine/Source/Editor/PropertyEditor/Private/SDetailsViewBase.cpp
--- a/Engine/Source/Editor/PropertyEditor/Private/SDetailsViewBase.cpp
+++ b/Engine/Source/Editor/PropertyEditor/Private/SDetailsViewBase.cpp
@@ -599,10 +599,6 @@ void SDetailsViewBase::UpdatePropertyMaps()
 
 	for(FDetailLayoutData& LayoutData : DetailLayouts)
 	{
-		// Check uniqueness.  It is critical that detail layouts can be destroyed
-		// We need to be able to create a new detail layout and properly clean up the old one in the process
-		check(!LayoutData.DetailLayout.IsValid() || LayoutData.DetailLayout.IsUnique());
-
 		// Allow customizations to perform cleanup as the delete occurs later on
 		for (TSharedPtr<IDetailCustomization>& DetailCustomization : LayoutData.CustomizationClassInstances)
 		{
