Restrict Conversation Reroute nodes to 1 Output link. Fixes multiple visual outputs even though it only fires the Leftmost output.
#rnx

diff --git a/Engine/Plugins/Experimental/CommonConversation/Source/CommonConversationGraph/Private/ConversationGraphSchema.cpp b/Engine/Plugins/Experimental/CommonConversation/Source/CommonConversationGraph/Private/ConversationGraphSchema.cpp
--- a/Engine/Plugins/Experimental/CommonConversation/Source/CommonConversationGraph/Private/ConversationGraphSchema.cpp
+++ b/Engine/Plugins/Experimental/CommonConversation/Source/CommonConversationGraph/Private/ConversationGraphSchema.cpp
@@ -40,6 +40,14 @@ namespace ConversationEditorCVar
 		TEXT("This cvar controles if the Conversation Editor should check for cycles when links are created.\n")
 		TEXT("0: Don't Check, 1: Check for Cycles (Default)"),
 		ECVF_Default);
+
+	static bool DisallowMultipleRerouteNodeOutputLinksCVar = false;
+	FAutoConsoleVariableRef CVarDisallowMultipleRerouteNodeOutputLinks(
+		TEXT("ConversationEditor.DisallowMultipleRerouteNodeOutputLinks"),
+		DisallowMultipleRerouteNodeOutputLinksCVar,
+		TEXT("Disallows Reroute nodes from visually splitting output links in Conversation Editor graph. Split links result in leftmost link always executing.\n")
+		TEXT("0: Allow multiple output links (Default), 1: Disallow multiple output links"),
+		ECVF_Default);
 }
 
 TSharedPtr<FGraphNodeClassHelper> ConversationClassCache;
@@ -247,6 +255,39 @@ const FPinConnectionResponse UConversationGraphSchema::CanCreateConnection(const
 		return FPinConnectionResponse(CONNECT_RESPONSE_DISALLOW, LOCTEXT("PinErrorSameNode", "Both are on the same node"));
 	}
 
+	// Check that both links are owned with a valid node class before using the class
+	if (!PinA->GetOwningNode())
+	{
+		if (PinA->Direction == EGPD_Input)
+		{
+			return FPinConnectionResponse(CONNECT_RESPONSE_DISALLOW, LOCTEXT("InputNodeTypeUnrecognized", "Input node type undefined"));
+		}
+		else if(PinA->Direction == EGPD_Output)
+		{
+			return FPinConnectionResponse(CONNECT_RESPONSE_DISALLOW, LOCTEXT("OutputNodeTypeUnrecognized", "Output node type undefined"));
+		}
+		else
+		{
+			return FPinConnectionResponse(CONNECT_RESPONSE_DISALLOW, LOCTEXT("NodeTypeUnrecognized", "Owning node type undefined"));
+		}
+	}
+
+	if (!PinB->GetOwningNode())
+	{
+		if (PinB->Direction == EGPD_Input)
+		{
+			return FPinConnectionResponse(CONNECT_RESPONSE_DISALLOW, LOCTEXT("InputNodeTypeUnrecognized", "Input node type undefined"));
+		}
+		else if (PinB->Direction == EGPD_Output)
+		{
+			return FPinConnectionResponse(CONNECT_RESPONSE_DISALLOW, LOCTEXT("OutputNodeTypeUnrecognized", "Output node type undefined"));
+		}
+		else
+		{
+			return FPinConnectionResponse(CONNECT_RESPONSE_DISALLOW, LOCTEXT("NodeTypeUnrecognized", "Owning node type undefined"));
+		}
+	}
+
 	const bool bPinAIsSingleComposite = (PinA->PinType.PinCategory == UConversationGraphTypes::PinCategory_SingleComposite);
 	const bool bPinAIsSingleTask = (PinA->PinType.PinCategory == UConversationGraphTypes::PinCategory_SingleTask);
 	const bool bPinAIsSingleNode = (PinA->PinType.PinCategory == UConversationGraphTypes::PinCategory_SingleNode);
@@ -346,37 +387,56 @@ const FPinConnectionResponse UConversationGraphSchema::CanCreateConnection(const
 	const bool bPinASingleLink = bPinAIsSingleComposite || bPinAIsSingleTask || bPinAIsSingleNode;
 	const bool bPinBSingleLink = bPinBIsSingleComposite || bPinBIsSingleTask || bPinBIsSingleNode;
 
+	// Joint Rules For Pins
+	//----------------------------------
+	//PinB is receiving input from other sources
 	if (PinB->Direction == EGPD_Input && PinB->LinkedTo.Num() > 0)
 	{
+		// PinA is exclusive output
 		if (bPinASingleLink)
 		{
+			// break all previous links between both nodes
 			return FPinConnectionResponse(CONNECT_RESPONSE_BREAK_OTHERS_AB, LOCTEXT("PinConnectReplace", "Replace connection"));
 		}
-		else
-		{
-			//@TODO: CONVERSATION: I think this is safe...
-			//return FPinConnectionResponse(CONNECT_RESPONSE_BREAK_OTHERS_B, LOCTEXT("PinConnectReplace", "Replace connection"));
-		}
 	}
 	else if (PinA->Direction == EGPD_Input && PinA->LinkedTo.Num() > 0)
 	{
+		// Pin B is exclusive output
 		if (bPinBSingleLink)
 		{
+			// break all previous links between both nodes
 			return FPinConnectionResponse(CONNECT_RESPONSE_BREAK_OTHERS_AB, LOCTEXT("PinConnectReplace", "Replace connection"));
 		}
-		else
+	}
+
+	// Singular Rules For Pins
+	//------------------------------
+	// Reroute Nodes have a single output link
+	// Not the same as being a SingleLink. Receiving nodes are still inclusive w/ unrestricted inputs
+
+	if(ConversationEditorCVar::DisallowMultipleRerouteNodeOutputLinksCVar)
+	{ 
+		if (PinA->GetOwningNode()->IsA(UConversationGraphNode_Knot::StaticClass()) && PinA->Direction == EGPD_Output)
+		{
+			return FPinConnectionResponse(CONNECT_RESPONSE_BREAK_OTHERS_A, LOCTEXT("PinRerouteOutputOverride", "Reroute node limited to 1 output link"));
+		}
+
+		if (PinB->GetOwningNode()->IsA(UConversationGraphNode_Knot::StaticClass()) && PinB->Direction == EGPD_Output)
 		{
-			//@TODO: CONVERSATION: I think this is safe...
-			//return FPinConnectionResponse(CONNECT_RESPONSE_BREAK_OTHERS_A, LOCTEXT("PinConnectReplace", "Replace connection"));
+			return FPinConnectionResponse(CONNECT_RESPONSE_BREAK_OTHERS_B, LOCTEXT("PinRerouteOutputOverride", "Reroute node limited to 1 output link"));
 		}
 	}
 
+	// Pin A is an exclusive link and is already linked to other sources
 	if (bPinASingleLink && PinA->LinkedTo.Num() > 0)
 	{
+		// break all previous links to pin A
 		return FPinConnectionResponse(CONNECT_RESPONSE_BREAK_OTHERS_A, LOCTEXT("PinConnectReplace", "Replace connection"));
 	}
 	else if (bPinBSingleLink && PinB->LinkedTo.Num() > 0)
 	{
+		// Pin B is an exclusive link and is already linked to other sources
+		// break all previous links to pin B
 		return FPinConnectionResponse(CONNECT_RESPONSE_BREAK_OTHERS_B, LOCTEXT("PinConnectReplace", "Replace connection"));
 	}
 
