Add ability for conversation nodes to control the nodes that they are allowed to connect to.

diff --git a/Engine/Plugins/Experimental/CommonConversation/Source/CommonConversationGraph/Private/ConversationGraphNode.cpp b/Engine/Plugins/Experimental/CommonConversation/Source/CommonConversationGraph/Private/ConversationGraphNode.cpp
--- a/Engine/Plugins/Experimental/CommonConversation/Source/CommonConversationGraph/Private/ConversationGraphNode.cpp
+++ b/Engine/Plugins/Experimental/CommonConversation/Source/CommonConversationGraph/Private/ConversationGraphNode.cpp
@@ -2,6 +2,7 @@
 
 #include "ConversationGraphNode.h"
 #include "ConversationGraphSchema.h"
+#include "ConversationGraphNode_Knot.h"
 
 #include "ConversationGraphTypes.h"
 #include "SConversationGraphNode.h"
@@ -123,6 +124,51 @@ FName UConversationGraphNode::GetNameIcon() const
 	return FName("BTEditor.Graph.BTNode.Icon");
 }
 
+bool UConversationGraphNode::IsOutBoundConnectionAllowed(const UConversationGraphNode* OtherNode, FText& OutErrorMessage) const
+{
+	if (!OtherNode)
+	{
+		return false;
+	}
+
+	if (const UConversationNodeWithLinks* MyTaskNode = GetRuntimeNode<UConversationNodeWithLinks>())
+	{
+		if (const UConversationNodeWithLinks* OtherTaskNode = OtherNode->GetRuntimeNode<UConversationNodeWithLinks>())
+		{
+			return MyTaskNode->IsOutBoundConnectionAllowed(OtherTaskNode, OutErrorMessage);
+		}
+	}
+
+	return true;
+}
+
+bool UConversationGraphNode::IsOutBoundConnectionAllowed(const UConversationGraphNode_Knot* KnotNode, FText& OutErrorMessage) const
+{
+	if (!KnotNode)
+	{
+		return false;
+	}
+
+	if (const UConversationNodeWithLinks* MyTaskNode = GetRuntimeNode<UConversationNodeWithLinks>())
+	{
+		TArray<UConversationGraphNode*> GraphNodes;
+		KnotNode->GatherAllOutBoundGraphNodes(GraphNodes);
+
+		for (const UConversationGraphNode* OtherNode : GraphNodes)
+		{
+			if (const UConversationNodeWithLinks* OtherTaskNode = OtherNode->GetRuntimeNode<UConversationNodeWithLinks>())
+			{
+				if (!MyTaskNode->IsOutBoundConnectionAllowed(OtherTaskNode, OutErrorMessage))
+				{
+					return false;
+				}
+			}
+		}
+	}
+
+	return true;
+}
+
 FText UConversationGraphNode::GetDescription() const
 {
 	if (const UConversationNode* RuntimeNode = Cast<const UConversationNode>(NodeInstance))
diff --git a/Engine/Plugins/Experimental/CommonConversation/Source/CommonConversationGraph/Private/ConversationGraphNode_Knot.cpp b/Engine/Plugins/Experimental/CommonConversation/Source/CommonConversationGraph/Private/ConversationGraphNode_Knot.cpp
--- a/Engine/Plugins/Experimental/CommonConversation/Source/CommonConversationGraph/Private/ConversationGraphNode_Knot.cpp
+++ b/Engine/Plugins/Experimental/CommonConversation/Source/CommonConversationGraph/Private/ConversationGraphNode_Knot.cpp
@@ -1,6 +1,8 @@
 // Copyright Epic Games, Inc. All Rights Reserved.
 
 #include "ConversationGraphNode_Knot.h"
+#include "ConversationGraphNode.h"
+#include "ConversationNode.h"
 #include "SGraphNodeKnot.h"
 
 #include UE_INLINE_GENERATED_CPP_BY_NAME(ConversationGraphNode_Knot)
@@ -93,6 +95,122 @@ TSharedPtr<SGraphNode> UConversationGraphNode_Knot::CreateVisualWidget()
 	return SNew(SGraphNodeKnot, this);
 }
 
+void UConversationGraphNode_Knot::GatherAllInBoundGraphNodes(TArray<UConversationGraphNode*>& OutGraphNodes) const
+{
+	TArray<const UConversationGraphNode_Knot*> VisitedKnots;
+	GatherAllInBoundGraphNodes_Internal(OutGraphNodes, VisitedKnots);
+}
+
+void UConversationGraphNode_Knot::GatherAllOutBoundGraphNodes(TArray<UConversationGraphNode*>& OutGraphNodes) const
+{
+	TArray<const UConversationGraphNode_Knot*> VisitedKnots;
+	GatherAllOutBoundGraphNodes_Internal(OutGraphNodes, VisitedKnots);
+}
+
+bool UConversationGraphNode_Knot::IsOutBoundConnectionAllowed(const UConversationGraphNode* OtherNode, FText& OutErrorMessage) const
+{
+	if (!OtherNode)
+	{
+		return false;
+	}
+
+	if (const UConversationNodeWithLinks* OtherTaskNode = OtherNode->GetRuntimeNode<UConversationNodeWithLinks>())
+	{
+		TArray<UConversationGraphNode*> GraphNodes;
+		GatherAllInBoundGraphNodes(GraphNodes);
+
+		for (const UConversationGraphNode* MyNode : GraphNodes)
+		{
+			if (const UConversationNodeWithLinks* MyTaskNode = MyNode->GetRuntimeNode<UConversationNodeWithLinks>())
+			{
+				if (!MyTaskNode->IsOutBoundConnectionAllowed(OtherTaskNode, OutErrorMessage))
+				{
+					return false;
+				}
+			}
+		}
+	}
+
+	return true;
+}
+
+bool UConversationGraphNode_Knot::IsOutBoundConnectionAllowed(const UConversationGraphNode_Knot* OtherKnotNode, FText& OutErrorMessage) const
+{
+	if (!OtherKnotNode)
+	{
+		return false;
+	}
+
+	TArray<UConversationGraphNode*> InGraphNodes;
+	GatherAllInBoundGraphNodes(InGraphNodes);
+	OtherKnotNode->GatherAllInBoundGraphNodes(InGraphNodes);
+
+	TArray<UConversationGraphNode*> OutGraphNodes;
+	GatherAllOutBoundGraphNodes(OutGraphNodes);
+	OtherKnotNode->GatherAllOutBoundGraphNodes(OutGraphNodes);
+
+	for (const UConversationGraphNode* MyNode : InGraphNodes)
+	{
+		if (const UConversationNodeWithLinks* MyTaskNode = MyNode->GetRuntimeNode<UConversationNodeWithLinks>())
+		{
+			for (const UConversationGraphNode* OtherNode : OutGraphNodes)
+			{
+				if (const UConversationNodeWithLinks* OtherTaskNode = OtherNode->GetRuntimeNode<UConversationNodeWithLinks>())
+				{
+					if (!MyTaskNode->IsOutBoundConnectionAllowed(OtherTaskNode, OutErrorMessage))
+					{
+						return false;
+					}
+				}
+			}
+		}
+	}
+
+	return true;
+}
+
+void UConversationGraphNode_Knot::GatherAllInBoundGraphNodes_Internal(TArray<UConversationGraphNode*>& OutGraphNodes, TArray<const UConversationGraphNode_Knot*>& VisitedKnots) const
+{
+	if (VisitedKnots.Contains(this))
+	{
+		return;
+	}
+	VisitedKnots.Add(this);
+
+	for (UEdGraphPin* LinkedPin : GetInputPin()->LinkedTo)
+	{
+		if (UConversationGraphNode* GraphNode = Cast<UConversationGraphNode>(LinkedPin->GetOwningNode()))
+		{
+			OutGraphNodes.Add(GraphNode);
+		}
+		else if (UConversationGraphNode_Knot* GraphKnot = Cast<UConversationGraphNode_Knot>(LinkedPin->GetOwningNode()))
+		{
+			GraphKnot->GatherAllInBoundGraphNodes_Internal(OutGraphNodes, VisitedKnots);
+		}
+	}
+}
+
+void UConversationGraphNode_Knot::GatherAllOutBoundGraphNodes_Internal(TArray<UConversationGraphNode*>& OutGraphNodes, TArray<const UConversationGraphNode_Knot*>& VisitedKnots) const
+{
+	if (VisitedKnots.Contains(this))
+	{
+		return;
+	}
+	VisitedKnots.Add(this);
+
+	for (UEdGraphPin* LinkedPin : GetOutputPin()->LinkedTo)
+	{
+		if (UConversationGraphNode* GraphNode = Cast<UConversationGraphNode>(LinkedPin->GetOwningNode()))
+		{
+			OutGraphNodes.Add(GraphNode);
+		}
+		else if (UConversationGraphNode_Knot* GraphKnot = Cast<UConversationGraphNode_Knot>(LinkedPin->GetOwningNode()))
+		{
+			GraphKnot->GatherAllOutBoundGraphNodes_Internal(OutGraphNodes, VisitedKnots);
+		}
+	}
+}
+
 /////////////////////////////////////////////////////
 
 #undef LOCTEXT_NAMESPACE
diff --git a/Engine/Plugins/Experimental/CommonConversation/Source/CommonConversationGraph/Private/ConversationGraphSchema.cpp b/Engine/Plugins/Experimental/CommonConversation/Source/CommonConversationGraph/Private/ConversationGraphSchema.cpp
--- a/Engine/Plugins/Experimental/CommonConversation/Source/CommonConversationGraph/Private/ConversationGraphSchema.cpp
+++ b/Engine/Plugins/Experimental/CommonConversation/Source/CommonConversationGraph/Private/ConversationGraphSchema.cpp
@@ -68,6 +68,65 @@ FGraphNodeClassHelper& GetConversationClassCache()
 	return *ConversationClassCache.Get();
 }
 
+bool IsConnectionAllowed(const UEdGraphPin* PinA, const UEdGraphPin* PinB, FText& OutErrorMessage)
+{
+	if (!PinA || !PinB)
+	{
+		return false;
+	}
+
+	const UConversationGraphNode* PinAGraphNode = Cast<UConversationGraphNode>(PinA->GetOwningNode());
+	const UConversationGraphNode_Knot* PinAKnot = Cast<UConversationGraphNode_Knot>(PinA->GetOwningNode());
+	const UConversationGraphNode* PinBGraphNode = Cast<UConversationGraphNode>(PinB->GetOwningNode());
+	const UConversationGraphNode_Knot* PinBKnot = Cast<UConversationGraphNode_Knot>(PinB->GetOwningNode());
+
+	// If both are GraphNode
+	if(PinAGraphNode && PinBGraphNode)
+	{
+		if (PinA->Direction == EGPD_Output)
+		{
+			return PinAGraphNode->IsOutBoundConnectionAllowed(PinBGraphNode, OutErrorMessage);
+		}
+		else if (PinB->Direction == EGPD_Output)
+		{
+			return PinBGraphNode->IsOutBoundConnectionAllowed(PinAGraphNode, OutErrorMessage);
+		}
+	}
+	// If both are Knot, direction does not matter
+	else if (PinAKnot && PinBKnot)
+	{
+		return PinAKnot->IsOutBoundConnectionAllowed(PinBKnot, OutErrorMessage);
+	}
+	// If one is GraphNode and one is Knot
+	else
+	{
+		if (PinA->Direction == EGPD_Output)
+		{
+			if (PinAGraphNode && PinBKnot)
+			{
+				return PinAGraphNode->IsOutBoundConnectionAllowed(PinBKnot, OutErrorMessage);
+			}
+			else if (PinAKnot && PinBGraphNode)
+			{
+				return PinAKnot->IsOutBoundConnectionAllowed(PinBGraphNode, OutErrorMessage);
+			}
+		}
+		else if (PinB->Direction == EGPD_Output)
+		{
+			if (PinBGraphNode && PinAKnot)
+			{
+				return PinBGraphNode->IsOutBoundConnectionAllowed(PinAKnot, OutErrorMessage);
+			}
+			else if (PinBKnot && PinAGraphNode)
+			{
+				return PinBKnot->IsOutBoundConnectionAllowed(PinAGraphNode, OutErrorMessage);
+			}
+		}
+	}
+
+	return true;
+}
+
 //////////////////////////////////////////////////////////////////////
 //
 
@@ -384,6 +443,17 @@ const FPinConnectionResponse UConversationGraphSchema::CanCreateConnection(const
 		}
 	}
 
+	// Check if the connection is allowed by the tasks
+	FText ErrorMessage;
+	if (!IsConnectionAllowed(PinA, PinB, ErrorMessage))
+	{
+		if (ErrorMessage.IsEmpty())
+		{
+			ErrorMessage = LOCTEXT("DefaultConnectionNotAllowed", "The connection between these nodes is not allowed");
+		}
+		return FPinConnectionResponse(CONNECT_RESPONSE_DISALLOW, MoveTemp(ErrorMessage));
+	}
+
 	const bool bPinASingleLink = bPinAIsSingleComposite || bPinAIsSingleTask || bPinAIsSingleNode;
 	const bool bPinBSingleLink = bPinBIsSingleComposite || bPinBIsSingleTask || bPinBIsSingleNode;
 
diff --git a/Engine/Plugins/Experimental/CommonConversation/Source/CommonConversationGraph/Public/ConversationGraphNode.h b/Engine/Plugins/Experimental/CommonConversation/Source/CommonConversationGraph/Public/ConversationGraphNode.h
--- a/Engine/Plugins/Experimental/CommonConversation/Source/CommonConversationGraph/Public/ConversationGraphNode.h
+++ b/Engine/Plugins/Experimental/CommonConversation/Source/CommonConversationGraph/Public/ConversationGraphNode.h
@@ -7,6 +7,7 @@
 
 namespace ENodeTitleType { enum Type : int; }
 
+class UConversationGraphNode_Knot;
 
 UCLASS()
 class COMMONCONVERSATIONGRAPH_API UConversationGraphNode : public UAIGraphNode
@@ -34,6 +35,9 @@ public:
 	/** gets icon resource name for title bar */
 	virtual FName GetNameIcon() const;
 
+	bool IsOutBoundConnectionAllowed(const UConversationGraphNode* OtherNode, FText& OutErrorMessage) const;
+	bool IsOutBoundConnectionAllowed(const UConversationGraphNode_Knot* KnotNode, FText& OutErrorMessage) const;
+
 	template<class T>
 	T* GetRuntimeNode() const
 	{
diff --git a/Engine/Plugins/Experimental/CommonConversation/Source/CommonConversationGraph/Public/ConversationGraphNode_Knot.h b/Engine/Plugins/Experimental/CommonConversation/Source/CommonConversationGraph/Public/ConversationGraphNode_Knot.h
--- a/Engine/Plugins/Experimental/CommonConversation/Source/CommonConversationGraph/Public/ConversationGraphNode_Knot.h
+++ b/Engine/Plugins/Experimental/CommonConversation/Source/CommonConversationGraph/Public/ConversationGraphNode_Knot.h
@@ -6,6 +6,7 @@
 #include "ConversationGraphNode_Knot.generated.h"
 
 class SGraphNode;
+class UConversationGraphNode;
 
 UCLASS(MinimalAPI)
 class UConversationGraphNode_Knot : public UEdGraphNode
@@ -37,6 +38,16 @@ public:
 	{
 		return Pins[1];
 	}
+
+	void GatherAllInBoundGraphNodes(TArray<UConversationGraphNode*>& OutGraphNodes) const;
+	void GatherAllOutBoundGraphNodes(TArray<UConversationGraphNode*>& OutGraphNodes) const;
+
+	bool IsOutBoundConnectionAllowed(const UConversationGraphNode* OtherNode, FText& OutErrorMessage) const;
+	bool IsOutBoundConnectionAllowed(const UConversationGraphNode_Knot* OtherKnotNode, FText& OutErrorMessage) const;
+
+private:
+	void GatherAllInBoundGraphNodes_Internal(TArray<UConversationGraphNode*>& OutGraphNodes, TArray<const UConversationGraphNode_Knot*>& VisitedKnots) const;
+	void GatherAllOutBoundGraphNodes_Internal(TArray<UConversationGraphNode*>& OutGraphNodes, TArray<const UConversationGraphNode_Knot*>& VisitedKnots) const;
 };
 
 #if UE_ENABLE_INCLUDE_ORDER_DEPRECATED_IN_5_2
diff --git a/Engine/Plugins/Experimental/CommonConversation/Source/CommonConversationRuntime/Public/ConversationNode.h b/Engine/Plugins/Experimental/CommonConversation/Source/CommonConversationRuntime/Public/ConversationNode.h
--- a/Engine/Plugins/Experimental/CommonConversation/Source/CommonConversationRuntime/Public/ConversationNode.h
+++ b/Engine/Plugins/Experimental/CommonConversation/Source/CommonConversationRuntime/Public/ConversationNode.h
@@ -175,6 +175,8 @@ public:
 	UPROPERTY()
 	TArray<FGuid> OutputConnections;
 
+	virtual bool IsOutBoundConnectionAllowed(const UConversationNodeWithLinks* OtherNode, FText& OutErrorMessage) const { return true; }
+
 	friend class UConversationRegistry;
 };
 
