Add Outbound node checks to the conversation graph asset validation.

diff --git a/Engine/Plugins/Experimental/CommonConversation/Source/CommonConversationRuntime/Private/ConversationDatabase.cpp b/Engine/Plugins/Experimental/CommonConversation/Source/CommonConversationRuntime/Private/ConversationDatabase.cpp
--- a/Engine/Plugins/Experimental/CommonConversation/Source/CommonConversationRuntime/Private/ConversationDatabase.cpp
+++ b/Engine/Plugins/Experimental/CommonConversation/Source/CommonConversationRuntime/Private/ConversationDatabase.cpp
@@ -132,7 +132,36 @@ EDataValidationResult UConversationDatabase::IsDataValid(FDataValidationContext&
 		}
 	}
 
-	return CombineDataValidationResults(SuperResult, Result);
+	EDataValidationResult OutBoundResult = ValidateOutBoundConnections(Context);
+
+	return CombineDataValidationResults(SuperResult, CombineDataValidationResults(Result, OutBoundResult));
+}
+
+EDataValidationResult UConversationDatabase::ValidateOutBoundConnections(class FDataValidationContext& Context) const
+{
+	EDataValidationResult Result = EDataValidationResult::Valid;
+	for (const TPair<FGuid, TObjectPtr<UConversationNode>>& Node : FullNodeMap)
+	{
+		if (const UConversationNodeWithLinks* ConversationNode = Cast<UConversationNodeWithLinks>(Node.Value))
+		{
+			for (const FGuid& OtherNodeGuid : ConversationNode->OutputConnections)
+			{
+				if (const UConversationNodeWithLinks* OtherConversationNode = Cast<UConversationNodeWithLinks>(FullNodeMap[OtherNodeGuid]))
+				{
+					FText ErrorMessage;
+					if (!ConversationNode->IsOutBoundConnectionAllowed(OtherConversationNode, ErrorMessage))
+					{
+						Context.AddError(FText::Format(LOCTEXT("OutBoundConnectionError", "({0}) -> ({1}) : {2}"),
+							ConversationNode->GetDisplayNameText(), OtherConversationNode->GetDisplayNameText(), ErrorMessage));
+
+						Result = EDataValidationResult::Invalid;
+					}
+				}
+			}
+		}
+	}
+
+	return Result;
 }
 #endif // WITH_EDITOR
 
diff --git a/Engine/Plugins/Experimental/CommonConversation/Source/CommonConversationRuntime/Public/ConversationDatabase.h b/Engine/Plugins/Experimental/CommonConversation/Source/CommonConversationRuntime/Public/ConversationDatabase.h
--- a/Engine/Plugins/Experimental/CommonConversation/Source/CommonConversationRuntime/Public/ConversationDatabase.h
+++ b/Engine/Plugins/Experimental/CommonConversation/Source/CommonConversationRuntime/Public/ConversationDatabase.h
@@ -80,6 +80,7 @@ public:
 #if WITH_EDITOR
 	virtual void PostDuplicate(bool bDuplicateForPIE) override;
 	virtual EDataValidationResult IsDataValid(class FDataValidationContext& Context) const override;
+	virtual EDataValidationResult ValidateOutBoundConnections(class FDataValidationContext& Context) const;
 	PRAGMA_DISABLE_DEPRECATION_WARNINGS // Suppress compiler warning on override of deprecated function
 	UE_DEPRECATED(5.0, "Use version that takes FObjectPreSaveContext instead.")
 	virtual void PreSave(const class ITargetPlatform* TargetPlatform) override;
