Add GetCurrentChoiceTags method to ConversationInstance

#rnx

diff --git a/Engine/Plugins/Experimental/CommonConversation/Source/CommonConversationRuntime/Public/ConversationInstance.h b/Engine/Plugins/Experimental/CommonConversation/Source/CommonConversationRuntime/Public/ConversationInstance.h
--- a/Engine/Plugins/Experimental/CommonConversation/Source/CommonConversationRuntime/Public/ConversationInstance.h
+++ b/Engine/Plugins/Experimental/CommonConversation/Source/CommonConversationRuntime/Public/ConversationInstance.h
@@ -132,6 +132,7 @@ private:
 #endif //WITH_SERVER_CODE
 
 protected:
+	const FGameplayTagContainer& GetCurrentChoiceTags() const { return CurrentBranchPoint.ClientChoice.ChoiceTags; }
 
 	TArray<FClientConversationOptionEntry> CurrentUserChoices;
 
