- Adding an editor accessor for the FullNodeMap

diff --git a/Engine/Plugins/Experimental/CommonConversation/Source/CommonConversationRuntime/Public/ConversationDatabase.h b/Engine/Plugins/Experimental/CommonConversation/Source/CommonConversationRuntime/Public/ConversationDatabase.h
--- a/Engine/Plugins/Experimental/CommonConversation/Source/CommonConversationRuntime/Public/ConversationDatabase.h
+++ b/Engine/Plugins/Experimental/CommonConversation/Source/CommonConversationRuntime/Public/ConversationDatabase.h
@@ -86,6 +86,8 @@ public:
 	virtual void PreSave(const class ITargetPlatform* TargetPlatform) override;
 	PRAGMA_ENABLE_DEPRECATION_WARNINGS
 	virtual void PreSave(FObjectPreSaveContext ObjectSaveContext) override;
+
+	const TMap<FGuid, TObjectPtr<UConversationNode>>& GetFullNodeMap() const { return FullNodeMap; }
 #endif
 
 	bool IsNodeReachable(const FGuid& NodeGUID) const { return ReachableNodeMap.Contains(NodeGUID); }
