Fix conversation occasionally ends before all nodes have been completed

- Also revert a bandaid fix on 36803547 that increases an input gate time

#rnx

diff --git a/Engine/Plugins/Experimental/CommonConversation/Source/CommonConversationRuntime/Private/ConversationInstance.cpp b/Engine/Plugins/Experimental/CommonConversation/Source/CommonConversationRuntime/Private/ConversationInstance.cpp
--- a/Engine/Plugins/Experimental/CommonConversation/Source/CommonConversationRuntime/Private/ConversationInstance.cpp
+++ b/Engine/Plugins/Experimental/CommonConversation/Source/CommonConversationRuntime/Private/ConversationInstance.cpp
@@ -11,6 +11,15 @@
 
 #include UE_INLINE_GENERATED_CPP_BY_NAME(ConversationInstance)
 
+namespace ConversationInstanceCVars
+{
+	static bool bShouldAbortConversationOnInvalidChoice = false;
+	static FAutoConsoleVariableRef CVarShouldAbortConversationOnInvalidChoice(
+		TEXT("Conversation.Instance.AbortConversationOnInvalidChoice"), bShouldAbortConversationOnInvalidChoice,
+		TEXT("About the conversation when an invalid choice is chosen."),
+		ECVF_Default);
+}
+
 //@TODO: CONVERSATION: Assert or otherwise guard all the Server* functions to only execute on the authority
 
 //////////////////////////////////////////////////////////////////////
@@ -295,8 +304,27 @@ void UConversationInstance::ServerAdvanceConversation(const FAdvanceConversation
 
 void UConversationInstance::OnInvalidBranchChoice(const FAdvanceConversationRequest& InChoicePicked)
 {
-	UE_LOG(LogCommonConversationRuntime, Error, TEXT("User picked option %s but it's not a legal output, aborting"), *InChoicePicked.ToString());
-	ServerAbortConversation();
+	if (ConversationInstanceCVars::bShouldAbortConversationOnInvalidChoice)
+	{
+		UE_LOG(LogCommonConversationRuntime, Error, TEXT("User picked option %s but it's not a legal output, aborting"), *InChoicePicked.ToString());
+		ServerAbortConversation();
+	}
+	else
+	{
+		UE_LOG(LogCommonConversationRuntime, Warning, TEXT("User picked option %s but it's not a legal output, ignoring"), *InChoicePicked.ToString());
+
+		// Forces the client to refresh current choices in case it is not aligned with the server
+		FConversationContext Context = FConversationContext::CreateServerContext(this, nullptr);
+		const bool bForceClientRefresh = true;
+
+		for (const FConversationParticipantEntry& ConversationParticipantEntry : GetParticipantListCopy())
+		{
+			if (UConversationParticipantComponent* ParticipantComponent = ConversationParticipantEntry.GetParticipantComponent())
+			{
+				ParticipantComponent->SendClientUpdatedChoices(Context, bForceClientRefresh);
+			}
+		}
+	}
 }
 
 void UConversationInstance::ServerAbortConversation()
diff --git a/Engine/Plugins/Experimental/CommonConversation/Source/CommonConversationRuntime/Private/ConversationParticipantComponent.cpp b/Engine/Plugins/Experimental/CommonConversation/Source/CommonConversationRuntime/Private/ConversationParticipantComponent.cpp
--- a/Engine/Plugins/Experimental/CommonConversation/Source/CommonConversationRuntime/Private/ConversationParticipantComponent.cpp
+++ b/Engine/Plugins/Experimental/CommonConversation/Source/CommonConversationRuntime/Private/ConversationParticipantComponent.cpp
@@ -198,11 +198,11 @@ void UConversationParticipantComponent::ClientUpdateConversations_Implementation
 	OnRep_ConversationsActive(OldConversationsActive);
 }
 
-void UConversationParticipantComponent::SendClientUpdatedChoices(const FConversationContext& Context)
+void UConversationParticipantComponent::SendClientUpdatedChoices(const FConversationContext& Context, bool bForcedRefresh)
 {
 #if WITH_SERVER_CODE
 	const TArray<FClientConversationOptionEntry> NewOptions = Context.GetActiveConversation()->GetCurrentUserConversationChoices();
-	if (NewOptions != LastMessage.Options)
+	if (NewOptions != LastMessage.Options || bForcedRefresh)
 	{
 		LastMessage.Options = NewOptions;
 		ClientUpdateConversation(LastMessage);
diff --git a/Engine/Plugins/Experimental/CommonConversation/Source/CommonConversationRuntime/Public/ConversationParticipantComponent.h b/Engine/Plugins/Experimental/CommonConversation/Source/CommonConversationRuntime/Public/ConversationParticipantComponent.h
--- a/Engine/Plugins/Experimental/CommonConversation/Source/CommonConversationRuntime/Public/ConversationParticipantComponent.h
+++ b/Engine/Plugins/Experimental/CommonConversation/Source/CommonConversationRuntime/Public/ConversationParticipantComponent.h
@@ -39,7 +39,7 @@ public:
 
 public:
 	void SendClientConversationMessage(const FConversationContext& Context, const FClientConversationMessagePayload& Payload);
-	virtual void SendClientUpdatedChoices(const FConversationContext& Context);
+	virtual void SendClientUpdatedChoices(const FConversationContext& Context, const bool bForcedRefresh = false);
 	void SendClientRefreshedTaskChoiceData(const FConversationNodeHandle& Handle, const FConversationContext& Context);
 
 	UFUNCTION(BlueprintCallable, Category=Conversation)
