Change check to ensure and add more info to it

#rnx

diff --git a/Engine/Plugins/Experimental/CommonConversation/Source/CommonConversationRuntime/Private/ConversationRegistry.cpp b/Engine/Plugins/Experimental/CommonConversation/Source/CommonConversationRuntime/Private/ConversationRegistry.cpp
--- a/Engine/Plugins/Experimental/CommonConversation/Source/CommonConversationRuntime/Private/ConversationRegistry.cpp
+++ b/Engine/Plugins/Experimental/CommonConversation/Source/CommonConversationRuntime/Private/ConversationRegistry.cpp
@@ -517,8 +517,11 @@ void UConversationRegistry::BuildDependenciesGraph()
 
 			for (FGuid& NodeId : NodeIds)
 			{
-				check(!NodeGuidToConversation.Contains(NodeId));
-				NodeGuidToConversation.Add(NodeId, ConversationDataAsset.ToSoftObjectPath());
+				if (const FSoftObjectPath* const ExistingEntry = NodeGuidToConversation.Find(NodeId);
+					ensureAlwaysMsgf(ExistingEntry == nullptr, TEXT("Trying to register NodeId %s to %s, but it's already registered to %s"), *NodeId.ToString(), *ConversationDataAsset.ToSoftObjectPath().ToString(), *ExistingEntry->ToString()))
+				{
+					NodeGuidToConversation.Add(NodeId, ConversationDataAsset.ToSoftObjectPath());
+				}
 			}
 		}
 	}
