Extending Conversation system with QOL functions

-adding a GetParticipantList function to UConversationInstance
-adding a GetOtherParticipantActors to UConversationParticipantComponent

#rnx

diff --git a/Engine/Plugins/Experimental/CommonConversation/Source/CommonConversationRuntime/Private/ConversationParticipantComponent.cpp b/Engine/Plugins/Experimental/CommonConversation/Source/CommonConversationRuntime/Private/ConversationParticipantComponent.cpp
--- a/Engine/Plugins/Experimental/CommonConversation/Source/CommonConversationRuntime/Private/ConversationParticipantComponent.cpp
+++ b/Engine/Plugins/Experimental/CommonConversation/Source/CommonConversationRuntime/Private/ConversationParticipantComponent.cpp
@@ -181,6 +181,44 @@ AActor* UConversationParticipantComponent::GetParticipantActor(const FGameplayTa
 	return nullptr;
 }
 
+void UConversationParticipantComponent::GetOtherParticipantActors(TArray<AActor*>& OutParticipantActors) const
+{
+	const TArray<FConversationParticipantEntry>* ConversationParticipantsPtr = nullptr;
+
+	if (GetOwner()->GetLocalRole() == ROLE_Authority)
+	{
+		if (Auth_CurrentConversation)
+		{
+			ConversationParticipantsPtr = &Auth_CurrentConversation->GetParticipantList();
+		}
+	}
+	else
+	{
+		ConversationParticipantsPtr = &LastMessage.Participants.List;
+	}
+
+	if (ConversationParticipantsPtr == nullptr)
+	{
+		return;
+	}
+
+	for (const FConversationParticipantEntry& ConversationParticipantEntry : *ConversationParticipantsPtr)
+	{
+		if (ConversationParticipantEntry.GetParticipantComponent() == this)
+		{
+			// Ignore self
+			continue;
+		}
+
+		if (!IsValid(ConversationParticipantEntry.Actor))
+		{
+			continue;
+		}
+
+		OutParticipantActors.Add(ConversationParticipantEntry.Actor);
+	}
+}
+
 void UConversationParticipantComponent::SendClientConversationMessage(const FConversationContext& Context, const FClientConversationMessagePayload& Payload)
 {
 #if WITH_SERVER_CODE
diff --git a/Engine/Plugins/Experimental/CommonConversation/Source/CommonConversationRuntime/Public/ConversationInstance.h b/Engine/Plugins/Experimental/CommonConversation/Source/CommonConversationRuntime/Public/ConversationInstance.h
--- a/Engine/Plugins/Experimental/CommonConversation/Source/CommonConversationRuntime/Public/ConversationInstance.h
+++ b/Engine/Plugins/Experimental/CommonConversation/Source/CommonConversationRuntime/Public/ConversationInstance.h
@@ -78,6 +78,11 @@ public:
 		return Participants.List;
 	}
 
+	const TArray<FConversationParticipantEntry>& GetParticipantList() const
+	{
+		return Participants.List;
+	}
+
 	FConversationParticipants GetParticipantsCopy() const
 	{
 		return Participants;
diff --git a/Engine/Plugins/Experimental/CommonConversation/Source/CommonConversationRuntime/Public/ConversationParticipantComponent.h b/Engine/Plugins/Experimental/CommonConversation/Source/CommonConversationRuntime/Public/ConversationParticipantComponent.h
--- a/Engine/Plugins/Experimental/CommonConversation/Source/CommonConversationRuntime/Public/ConversationParticipantComponent.h
+++ b/Engine/Plugins/Experimental/CommonConversation/Source/CommonConversationRuntime/Public/ConversationParticipantComponent.h
@@ -85,6 +85,9 @@ public:
 	UFUNCTION(BlueprintPure, Category = Conversation)
 	UE_API AActor* GetParticipantActor(const FGameplayTag& ParticipantTag) const;
 
+	UFUNCTION(BlueprintCallable, Category = Conversation)
+	UE_API void GetOtherParticipantActors(TArray<AActor*>& OutParticipantActors) const;
+
 protected:
 	UFUNCTION(Server, Reliable)
 	UE_API void ServerAdvanceConversation(const FAdvanceConversationRequest& InChoicePicked);
