Add warning log for when a Conversation Graph subnode is invalid

- Preferring a log over ensureAlwaysMsgf so we can learn about all faulty nodes without doing the full Crash Reporter flow every time (a single ensure hit should be enough to raise the bells)

#rnx
#Synced-CL 49670061

diff --git a/Engine/Plugins/Experimental/CommonConversation/Source/CommonConversationGraph/Private/ConversationCompiler.cpp b/Engine/Plugins/Experimental/CommonConversation/Source/CommonConversationGraph/Private/ConversationCompiler.cpp
--- a/Engine/Plugins/Experimental/CommonConversation/Source/CommonConversationGraph/Private/ConversationCompiler.cpp
+++ b/Engine/Plugins/Experimental/CommonConversation/Source/CommonConversationGraph/Private/ConversationCompiler.cpp
@@ -1,6 +1,8 @@
 // Copyright Epic Games, Inc. All Rights Reserved.
 
 #include "ConversationCompiler.h"
+
+#include "CommonConversationGraphModule.h"
 #include "ConversationGraph.h"
 #include "ConversationGraphSchema.h"
 #include "ConversationDatabase.h"
@@ -238,6 +240,14 @@ void FConversationCompiler::RebuildBank(UConversationDatabase* ConversationAsset
 					{
 						TaskNode->SubNodes.Add(TaskSubNode);
 					}
+					else
+					{
+						UE_LOG(LogCommonConversationGraph, Warning, TEXT("Invalid TaskSubNode inside Conversation Asset %s"), *ConversationAsset->GetPrimaryAssetId().ToString());
+					}
+				}
+				else
+				{
+					UE_LOG(LogCommonConversationGraph, Warning, TEXT("Invalid TypedSubNode inside Conversation Asset %s"), *ConversationAsset->GetPrimaryAssetId().ToString());
 				}
 			}
 		}
