Add an option to stop and manually build the conversation registry dependencies graph when needed, instead of waiting on a delegate

diff --git a/Engine/Plugins/Experimental/CommonConversation/Source/CommonConversationRuntime/Private/ConversationRegistry.cpp b/Engine/Plugins/Experimental/CommonConversation/Source/CommonConversationRuntime/Private/ConversationRegistry.cpp
--- a/Engine/Plugins/Experimental/CommonConversation/Source/CommonConversationRuntime/Private/ConversationRegistry.cpp
+++ b/Engine/Plugins/Experimental/CommonConversation/Source/CommonConversationRuntime/Private/ConversationRegistry.cpp
@@ -188,11 +188,19 @@ void UConversationRegistry::GameFeatureStateModified()
 
 void UConversationRegistry::OnGameFeatureActivated(const UGameFeatureData* GameFeatureData, const FString& PluginURL)
 {
+	if (bWaitForDependenciesGraphFlush)
+	{
+		return;
+	}
 	GameFeatureStateModified();
 }
 
 void UConversationRegistry::OnGameFeatureDeactivating(const UGameFeatureData* GameFeatureData, FGameFeatureDeactivatingContext& Context, const FString& PluginURL)
 {
+	if (bWaitForDependenciesGraphFlush)
+	{
+		return;
+	}
 	GameFeatureStateModified();
 }
 
@@ -326,6 +334,20 @@ TArray<FGuid> UConversationRegistry::GetOutputLinkGUIDs(const UConversationDatab
 	return Result;
 }
 
+void UConversationRegistry::FlushBuildDependenciesGraph()
+{
+	if (bWaitForDependenciesGraphFlush)
+	{
+		GameFeatureStateModified();
+		bWaitForDependenciesGraphFlush = false;
+	}
+}
+
+void UConversationRegistry::SetWaitDependenciesGraphToFlush(bool bWaitToFlush)
+{
+	bWaitForDependenciesGraphFlush = bWaitToFlush;
+}
+
 UConversationDatabase* UConversationRegistry::GetConversationFromNodeGUID(const FGuid& NodeGUID) const
 {
 	if (const FSoftObjectPath* ConversationPathPtr = NodeGuidToConversation.Find(NodeGUID))
diff --git a/Engine/Plugins/Experimental/CommonConversation/Source/CommonConversationRuntime/Public/ConversationRegistry.h b/Engine/Plugins/Experimental/CommonConversation/Source/CommonConversationRuntime/Public/ConversationRegistry.h
--- a/Engine/Plugins/Experimental/CommonConversation/Source/CommonConversationRuntime/Public/ConversationRegistry.h
+++ b/Engine/Plugins/Experimental/CommonConversation/Source/CommonConversationRuntime/Public/ConversationRegistry.h
@@ -107,6 +107,9 @@ public:
 
 	FAvailableConversationsChangedEvent AvailableConversationsChanged;
 
+	UE_API void FlushBuildDependenciesGraph();
+	UE_API void SetWaitDependenciesGraphToFlush(bool bWaitToFlush);
+	
 private:
 	UConversationDatabase* GetConversationFromNodeGUID(const FGuid& NodeGUID) const;
 
@@ -127,6 +130,7 @@ private:
 
 private:
 	bool bDependenciesBuilt = false;
+	bool bWaitForDependenciesGraphFlush = false;
 };
 
 #undef UE_API
