[AI][Conversation] Add ConversationToolset plugin

Add a standalone plugin for AI-callable inspection of ConversationDatabase assets. This enables MCP agents to traverse conversation graphs, read node properties, and inspect sub-node requirements without any custom Python code.

Changes:
- New ConversationToolset plugin with ConversationTools toolset (7 tools: list_entry_points, list_speakers, get_all_nodes, get_node_guids, get_node_by_guid, get_node_connections, get_sub_nodes)
- Add BlueprintType/BlueprintReadOnly metadata to ConversationDatabase.h, ConversationNode.h, ConversationTaskNode.h to expose conversation structures to Python
- Add BlueprintReadOnly + AllowPrivateAccess to UConversationNode::NodeName for direct Python attribute access

#Synced-CL 51880525
#Synced-CL 51880525

diff --git a/Engine/Plugins/Experimental/CommonConversation/Source/CommonConversationRuntime/Public/ConversationDatabase.h b/Engine/Plugins/Experimental/CommonConversation/Source/CommonConversationRuntime/Public/ConversationDatabase.h
--- a/Engine/Plugins/Experimental/CommonConversation/Source/CommonConversationRuntime/Public/ConversationDatabase.h
+++ b/Engine/Plugins/Experimental/CommonConversation/Source/CommonConversationRuntime/Public/ConversationDatabase.h
@@ -22,18 +22,18 @@ class UConversationRegistry;
  * all of those nodes with the same matching tag name, so that the entry point is
  * effectively randomized when there are multiple copies.
  */
-USTRUCT()
+USTRUCT(BlueprintType)
 struct FConversationEntryList
 {
 	GENERATED_BODY()
 
-	UPROPERTY()
+	UPROPERTY(BlueprintReadOnly, Category = Conversation)
 	FGameplayTag EntryTag;
 
-	UPROPERTY()
+	UPROPERTY(BlueprintReadOnly, Category = Conversation)
 	TArray<FGuid> DestinationList;
 
-	UPROPERTY()
+	UPROPERTY(BlueprintReadOnly, Category = Conversation)
 	FString EntryIdentifier;
 };
 
@@ -45,12 +45,12 @@ struct FConversationEntryList
 // (e.g., mapping a logical Player to the current player pawn)
 //
 
-USTRUCT()
+USTRUCT(BlueprintType)
 struct FCommonDialogueBankParticipant
 {
 	GENERATED_BODY()
 
-	UPROPERTY()
+	UPROPERTY(BlueprintReadOnly, Category = Conversation)
 	FText FallbackName;
 
 	/** Identifier represented by the component */
@@ -98,15 +98,15 @@ private: // Compiled Data
 	int32 CompilerVersion = INDEX_NONE;
 
 	// Compiled: Reachable nodes
-	UPROPERTY()
+	UPROPERTY(BlueprintReadOnly, Category = Conversation, meta = (AllowPrivateAccess = "true"))
 	TMap<FGuid, TObjectPtr<UConversationNode>> ReachableNodeMap;
 
 	// Compiled: Entry points
-	UPROPERTY(AssetRegistrySearchable)
+	UPROPERTY(AssetRegistrySearchable, BlueprintReadOnly, Category = Conversation, meta = (AllowPrivateAccess = "true"))
 	TArray<FConversationEntryList> EntryTags;
 
-	// Compiled: 
-	UPROPERTY(AssetRegistrySearchable)
+	// Compiled:
+	UPROPERTY(AssetRegistrySearchable, BlueprintReadOnly, Category = Conversation, meta = (AllowPrivateAccess = "true"))
 	FGameplayTagContainer ExitTags;
 	
 	// Compiled: 
@@ -119,7 +119,7 @@ private: // Compiled Data
 
 private:
 	// List of participant slots
-	UPROPERTY(EditAnywhere, Category=Conversation)
+	UPROPERTY(EditAnywhere, BlueprintReadOnly, Category=Conversation, meta = (AllowPrivateAccess = "true"))
 	TArray<FCommonDialogueBankParticipant> Speakers;
 
 private:
diff --git a/Engine/Plugins/Experimental/CommonConversation/Source/CommonConversationRuntime/Public/ConversationNode.h b/Engine/Plugins/Experimental/CommonConversation/Source/CommonConversationRuntime/Public/ConversationNode.h
--- a/Engine/Plugins/Experimental/CommonConversation/Source/CommonConversationRuntime/Public/ConversationNode.h
+++ b/Engine/Plugins/Experimental/CommonConversation/Source/CommonConversationRuntime/Public/ConversationNode.h
@@ -138,13 +138,13 @@ public:
 
 protected:
 	/** node name */
-	UPROPERTY(EditAnywhere, Category=Description, AdvancedDisplay)
+	UPROPERTY(EditAnywhere, BlueprintReadOnly, Category=Description, AdvancedDisplay, meta = (AllowPrivateAccess = "true"))
 	FString NodeName;
 
 	/**
 	 * The node's unique ID.  This value is set during compilation.
 	 */
-	UPROPERTY()
+	UPROPERTY(BlueprintReadOnly, Category = Conversation, meta = (AllowPrivateAccess = "true"))
 	FGuid Compiled_NodeGUID;
 
 #if WITH_EDITORONLY_DATA
@@ -174,7 +174,7 @@ class UConversationNodeWithLinks : public UConversationNode
 	GENERATED_BODY()
 
 public:
-	UPROPERTY()
+	UPROPERTY(BlueprintReadOnly, Category = Conversation)
 	TArray<FGuid> OutputConnections;
 
 	virtual bool IsOutBoundConnectionAllowed(const UConversationNodeWithLinks* OtherNode, FText& OutErrorMessage) const { return true; }
diff --git a/Engine/Plugins/Experimental/CommonConversation/Source/CommonConversationRuntime/Public/ConversationTaskNode.h b/Engine/Plugins/Experimental/CommonConversation/Source/CommonConversationRuntime/Public/ConversationTaskNode.h
--- a/Engine/Plugins/Experimental/CommonConversation/Source/CommonConversationRuntime/Public/ConversationTaskNode.h
+++ b/Engine/Plugins/Experimental/CommonConversation/Source/CommonConversationRuntime/Public/ConversationTaskNode.h
@@ -46,7 +46,7 @@ class UConversationTaskNode : public UConversationNodeWithLinks
 
 public:
 	// Requirements and side effects
-	UPROPERTY()
+	UPROPERTY(BlueprintReadOnly, Category = Conversation)
 	TArray<TObjectPtr<UConversationSubNode>> SubNodes;
 
 #if WITH_EDITORONLY_DATA
diff --git a/Engine/Plugins/Experimental/ConversationToolset/ConversationToolset.uplugin b/Engine/Plugins/Experimental/ConversationToolset/ConversationToolset.uplugin
new file mode 100644
--- /dev/null
+++ b/Engine/Plugins/Experimental/ConversationToolset/ConversationToolset.uplugin
@@ -0,0 +1,38 @@
+{
+	"FileVersion": 3,
+	"Version": 1,
+	"VersionName": "1.0",
+	"FriendlyName": "ConversationToolset",
+	"Description": "Toolset for Conversation Systems",
+	"Category": "Other",
+	"CreatedBy": "Epic Games",
+	"CreatedByURL": "",
+	"DocsURL": "",
+	"MarketplaceURL": "",
+	"SupportURL": "",
+	"EditorOnly": true,
+	"CanContainContent": true,
+	"IsBetaVersion": false,
+	"IsExperimentalVersion": true,
+	"NoRedist": true,
+	"Installed": false,
+	"EnabledByDefault": true,
+	"Modules": [
+		{
+			"Name": "ConversationToolset",
+			"Type": "Editor",
+			"LoadingPhase": "Default",
+			"TargetAllowList": ["Editor"]
+		}
+	],
+	"Plugins": [
+		{
+			"Name": "CommonConversation",
+			"Enabled": true
+		},
+		{
+			"Name": "ToolsetRegistry",
+			"Enabled": true
+		}
+	]
+}
diff --git a/Engine/Plugins/Experimental/ConversationToolset/Source/ConversationToolset/ConversationToolset.Build.cs b/Engine/Plugins/Experimental/ConversationToolset/Source/ConversationToolset/ConversationToolset.Build.cs
new file mode 100644
--- /dev/null
+++ b/Engine/Plugins/Experimental/ConversationToolset/Source/ConversationToolset/ConversationToolset.Build.cs
@@ -0,0 +1,12 @@
+// Copyright Epic Games, Inc. All Rights Reserved.
+
+using UnrealBuildTool;
+
+public class ConversationToolset : ModuleRules
+{
+	public ConversationToolset(ReadOnlyTargetRules Target) : base(Target)
+	{
+		PCHUsage = PCHUsageMode.UseExplicitOrSharedPCHs;
+		PublicDependencyModuleNames.Add("Core");
+	}
+}
diff --git a/Engine/Plugins/Experimental/ConversationToolset/Source/ConversationToolset/Private/ConversationToolset.cpp b/Engine/Plugins/Experimental/ConversationToolset/Source/ConversationToolset/Private/ConversationToolset.cpp
new file mode 100644
--- /dev/null
+++ b/Engine/Plugins/Experimental/ConversationToolset/Source/ConversationToolset/Private/ConversationToolset.cpp
@@ -0,0 +1,16 @@
+// Copyright Epic Games, Inc. All Rights Reserved.
+
+#include "ConversationToolset.h"
+
+#include "Modules/ModuleManager.h"
+
+/*virtual*/ void FConversationToolsetModule::StartupModule()
+{
+}
+
+/*virtual*/ void FConversationToolsetModule::ShutdownModule()
+{
+}
+
+IMPLEMENT_MODULE(FConversationToolsetModule, ConversationToolset);
+DEFINE_LOG_CATEGORY(LogConversationToolset);
diff --git a/Engine/Plugins/Experimental/ConversationToolset/Source/ConversationToolset/Public/ConversationToolset.h b/Engine/Plugins/Experimental/ConversationToolset/Source/ConversationToolset/Public/ConversationToolset.h
new file mode 100644
--- /dev/null
+++ b/Engine/Plugins/Experimental/ConversationToolset/Source/ConversationToolset/Public/ConversationToolset.h
@@ -0,0 +1,15 @@
+// Copyright Epic Games, Inc. All Rights Reserved.
+
+#pragma once
+
+#include "Modules/ModuleInterface.h"
+#include "Logging/LogMacros.h"
+
+DECLARE_LOG_CATEGORY_EXTERN(LogConversationToolset, Log, Log);
+
+class FConversationToolsetModule : public IModuleInterface
+{
+public:
+	virtual void StartupModule() override;
+	virtual void ShutdownModule() override;
+};
